Initial prototype of data interfaces and data interface graphs

Data interfaces are designed to be the basis for a functional, compositional framework, allowing loose bindings (a typed return value) between functional units (data interfaces). They incorporate the following main features:
- Simple call to 'get a value' from a data interface.
- Most communication is done via the UE::DataInterface::FContext, which handles data management.
- Type-erased values (UE::DataInterface::FParam) for results, parameters and state. These can be promoted to typed values (UE::DataInterface::TParam<T>) with runtime validation.
- 'Hidden state', allocated on-demand via a context, to support functional units that require statefulness. State is allocated according to the calling context, currently a hash of the 'callstack' and call site, so subsequent calls to the same data interface can reuse state.
- Batch processing of multiple state representations at once (e.g. processing multiple characters at one), via a chunked allocation strategy.

Also includes data interface graphs - a proof-of-concept implementation using RigVM to construct scripted data interface logic. This is still very early!

Still lots to do:
- Removal of typed interfaces (e.g. IDataInterface_Float)
- Validation of hidden state approach on a wider scale
- Validation of kernel processing approach
- Chunked branching (probably via a masking strategy) so branches can be handled within chunks
- Reworked type system incorporating type promotion/conversion
- UI work to allow for mixing instanced sub-objects and asset references with using TScriptInterface<IDataInterface>
- Expansion or RigVM usage and a ton of UI/UX work to get graph editing up to scratch
- Many many more things

diff --git a/Engine/Plugins/Experimental/DataInterface/DataInterface.uplugin b/Engine/Plugins/Experimental/DataInterface/DataInterface.uplugin
new file mode 100644
--- /dev/null
+++ b/Engine/Plugins/Experimental/DataInterface/DataInterface.uplugin
@@ -0,0 +1,25 @@
+{
+	"FileVersion" : 3,
+	"Version" : 1,
+	"VersionName" : "0.1",
+	"FriendlyName" : "Data Interface",
+	"Description" : "Framework for defining functional data flow",
+	"Category" : "Animation",
+	"CreatedBy" : "Epic Games, Inc.",
+	"CreatedByURL" : "http://epicgames.com",
+	"DocsURL" : "",
+	"MarketplaceURL" : "",
+	"SupportURL" : "",
+	"EnabledByDefault" : false,
+	"CanContainContent" : false,
+	"IsBetaVersion" : true,
+	"Installed" : false,
+	"Modules" :
+	[
+		{
+			"Name" : "DataInterface",
+			"Type" : "Runtime",
+			"LoadingPhase" : "PreDefault"
+		}
+	]
+}
diff --git a/Engine/Plugins/Experimental/DataInterface/Source/DataInterface/DataInterface.Build.cs b/Engine/Plugins/Experimental/DataInterface/Source/DataInterface/DataInterface.Build.cs
new file mode 100644
--- /dev/null
+++ b/Engine/Plugins/Experimental/DataInterface/Source/DataInterface/DataInterface.Build.cs
@@ -0,0 +1,19 @@
+// Copyright Epic Games, Inc. All Rights Reserved.
+
+namespace UnrealBuildTool.Rules
+{
+	public class DataInterface : ModuleRules
+	{
+		public DataInterface(ReadOnlyTargetRules Target) : base(Target)
+		{	
+			PublicDependencyModuleNames.AddRange(
+				new string[]
+				{
+					"Core",
+					"CoreUObject",
+				}
+			);
+
+		}
+	}
+}
\ No newline at end of file
diff --git a/Engine/Plugins/Experimental/DataInterface/Source/DataInterface/Private/DataInterface.cpp b/Engine/Plugins/Experimental/DataInterface/Source/DataInterface/Private/DataInterface.cpp
new file mode 100644
--- /dev/null
+++ b/Engine/Plugins/Experimental/DataInterface/Source/DataInterface/Private/DataInterface.cpp
@@ -0,0 +1,47 @@
+// Copyright Epic Games, Inc. All Rights Reserved.
+
+#include "DataInterface.h"
+#include "IDataInterface.h"
+#include "DataInterfaceContext.h"
+
+namespace UE::DataInterface::Private
+{
+
+static bool CheckCompatibility(const IDataInterface* InDataInterface, const UE::DataInterface::FContext& InContext)
+{
+	check(InDataInterface);
+	return InContext.GetResultRaw().GetType().GetName() == InDataInterface->GetReturnTypeName();
+}
+
+}
+
+bool IDataInterface::GetDataIfCompatibleInternal(const UE::DataInterface::FContext& InContext) const
+{
+	if(UE::DataInterface::Private::CheckCompatibility(this, InContext))
+	{
+		return GetDataRawInternal(InContext);
+	}
+
+	return false;
+}
+
+bool IDataInterface::GetData(const UE::DataInterface::FContext& Context, UE::DataInterface::FParam& OutResult) const
+{
+	const UE::DataInterface::FContext CallingContext = Context.WithResult(OutResult);
+	return GetDataIfCompatibleInternal(CallingContext);
+}
+
+bool IDataInterface::GetDataChecked(const UE::DataInterface::FContext& Context, UE::DataInterface::FParam& OutResult) const
+{
+	const UE::DataInterface::FContext CallingContext = Context.WithResult(OutResult);
+	check(UE::DataInterface::Private::CheckCompatibility(this, CallingContext));
+	return GetDataRawInternal(CallingContext);
+}
+
+bool IDataInterface::GetDataRawInternal(const UE::DataInterface::FContext& InContext) const
+{
+	// TODO: debug recording here on context
+
+	const UE::DataInterface::FContext CallContext = InContext.WithCallRaw(this);
+	return GetDataImpl(CallContext);
+}
diff --git a/Engine/Plugins/Experimental/DataInterface/Source/DataInterface/Private/DataInterfaceContext.cpp b/Engine/Plugins/Experimental/DataInterface/Source/DataInterface/Private/DataInterfaceContext.cpp
new file mode 100644
--- /dev/null
+++ b/Engine/Plugins/Experimental/DataInterface/Source/DataInterface/Private/DataInterfaceContext.cpp
@@ -0,0 +1,152 @@
+// Copyright Epic Games, Inc. All Rights Reserved.
+
+#include "DataInterfaceContext.h"
+#include "IDataInterface.h"
+#include "IDataInterfaceParameters.h"
+#include "DataInterfaceTypes.h"
+
+namespace UE::DataInterface
+{
+
+struct FThreadContextData : TThreadSingleton<FThreadContextData>
+{
+	TArray<const FContext*> ContextStack;
+};
+
+// @TODO: we might need to do this for all FContext constructions?
+// As parameters are gathered using the FContext linked list, we always need to access the top of the stack to get 
+// correct parameters. Shortcutting to the bottom of the stack via FThreadContext misses intervening parameters
+FThreadContext::FThreadContext(const FContext& InContext)
+{
+	FThreadContextData& ContextData = FThreadContextData::Get();
+	ContextData.ContextStack.Push(&InContext);
+}
+
+FThreadContext::~FThreadContext()
+{
+	FThreadContextData& ContextData = FThreadContextData::Get();
+	ContextData.ContextStack.Pop();
+}
+
+const FContext& FThreadContext::Get()
+{
+	const FThreadContextData& ContextData = FThreadContextData::Get();
+	return *ContextData.ContextStack.Top();
+}
+
+FContext::FContext(float InDeltaTime, FState& InState, IDataInterfaceParameters* InParameters)
+	: State(InState)
+	, Result(nullptr)
+	, Parameters(InParameters)
+	, DeltaTime(InDeltaTime)
+{
+}
+
+FContext FContext::WithResult(FParam& InResult) const
+{
+	FContext NewContext(DeltaTime, State, &InResult);
+	NewContext.Parent = this;
+	NewContext.Root = (Root == nullptr) ? this : Root;
+	NewContext.CallstackHash = CallstackHash;
+
+	return NewContext;
+}
+
+FContext FContext::WithParameter(FName ParameterId, const FParam& InParameter) const
+{
+	FContext NewContext(DeltaTime, State, Result);
+	NewContext.AdditionalParameters.Add(ParameterId, InParameter);
+	NewContext.Parent = this;
+	NewContext.Root = (Root == nullptr) ? this : Root;
+	NewContext.CallstackHash = CallstackHash;
+
+	return NewContext;
+}
+
+FContext FContext::WithParameters(TArrayView<const TPair<FName, FParam>> InParameters) const
+{
+	FContext NewContext(DeltaTime, State, Result);
+	for(const TPair<FName, FParam>& ParamPair : InParameters)
+	{
+		NewContext.AdditionalParameters.Add(ParamPair.Key, ParamPair.Value);
+	}
+	NewContext.Parent = this;
+	NewContext.Root = (Root == nullptr) ? this : Root;
+	NewContext.CallstackHash = CallstackHash;
+
+	return NewContext;
+}
+
+FContext FContext::WithParameters(IDataInterfaceParameters* InParameters) const
+{
+	FContext NewContext(DeltaTime, State, Result);
+	NewContext.Parameters = InParameters;
+	NewContext.Parent = this;
+	NewContext.Root = (Root == nullptr) ? this : Root;
+	NewContext.CallstackHash = CallstackHash;
+
+	return NewContext;
+}
+
+bool FContext::GetParameter(FName InKey, FParam& OutParam) const
+{
+	// Check parent contexts in turn. Parameters can be overriden by each scoped context.
+	for(const FContext* CurrentContext = this; CurrentContext != nullptr; CurrentContext = CurrentContext->Parent)
+	{
+		// Check parameter provider
+		if(CurrentContext->Parameters)
+		{
+			if(CurrentContext->Parameters->GetParameter(InKey, OutParam))
+			{
+				return true;
+			}
+		}
+		// Check additional parameters
+		// Note this is currently 'else if' because we dont allow creating new contexts with both a IDataInterfaceParameters
+		// interface and AdditionalParameters
+		else if(!CurrentContext->AdditionalParameters.IsEmpty())
+		{
+			// Find the parameter
+			if(const FParam* FoundParameter = CurrentContext->AdditionalParameters.Find(InKey))
+			{
+				// Check type compatibility
+				if(FoundParameter->CanAssignTo(OutParam))
+				{
+					OutParam = *FoundParameter;
+					return true;
+				}
+			}
+		}
+	}
+
+	return false;
+}
+
+TParam<const float> FContext::GetDeltaTimeParam() const
+{
+	return TWrapParam<const float>(*this, &DeltaTime);
+}
+
+FContext FContext::WithCallRaw(const IDataInterface* InDataInterface) const
+{
+	FContext NewContext(DeltaTime, State, Result);
+	NewContext.Parent = this;
+	NewContext.Root = (Root == nullptr) ? this : Root;
+	NewContext.CallstackHash = HashCombineFast(CallstackHash, GetTypeHash(InDataInterface));
+	return NewContext;
+}
+
+void FContext::FlushRelevancy() const
+{
+	// Flush all relevancy-based allocations if they were not used this update
+	// @TODO: this could be more efficient if we use a more linear iteration here 
+	for(auto It = State.RelevancyValueMap.CreateIterator(); It; ++It)
+	{
+		if(It->Value.UpdateCounter != UpdateCounter)
+		{
+			It.RemoveCurrent();
+		}
+	}
+}
+
+}
\ No newline at end of file
diff --git a/Engine/Plugins/Experimental/DataInterface/Source/DataInterface/Private/DataInterfaceModule.cpp b/Engine/Plugins/Experimental/DataInterface/Source/DataInterface/Private/DataInterfaceModule.cpp
new file mode 100644
--- /dev/null
+++ b/Engine/Plugins/Experimental/DataInterface/Source/DataInterface/Private/DataInterfaceModule.cpp
@@ -0,0 +1,26 @@
+// Copyright Epic Games, Inc. All Rights Reserved.
+
+#include "CoreMinimal.h"
+#include "Modules/ModuleInterface.h"
+#include "Modules/ModuleManager.h"
+#include "Misc/CoreDelegates.h"
+#include "DataInterfaceParam.h"
+
+namespace UE::DataInterface
+{
+
+class FModule : public IModuleInterface
+{
+public:
+	virtual void StartupModule() override
+	{
+		FCoreDelegates::OnPostEngineInit.AddLambda([]()
+		{
+			FParamType::FRegistrar::RegisterDeferredTypes();
+		});
+	}
+};
+
+}
+
+IMPLEMENT_MODULE(UE::DataInterface::FModule, DataInterface)
\ No newline at end of file
diff --git a/Engine/Plugins/Experimental/DataInterface/Source/DataInterface/Private/DataInterfaceParam.cpp b/Engine/Plugins/Experimental/DataInterface/Source/DataInterface/Private/DataInterfaceParam.cpp
new file mode 100644
--- /dev/null
+++ b/Engine/Plugins/Experimental/DataInterface/Source/DataInterface/Private/DataInterfaceParam.cpp
@@ -0,0 +1,143 @@
+﻿// Copyright Epic Games, Inc. All Rights Reserved.
+
+#include "DataInterfaceParam.h"
+
+namespace UE::DataInterface
+{
+
+static FParamType DefaultType;
+static TArray<FParamType*> GTypeRegistry;
+static TArray<TUniqueFunction<void(void)>> GDeferredTypes;
+
+FParam::FParam(const FParamType& InType, void* InData, EFlags InFlags)
+	: Data(InData)
+	, TypeId(InType.GetTypeId())
+	, Flags(InFlags)
+{
+	check(TypeId != INDEX_NONE);
+	check(Data);
+}
+
+FParam::FParam(const FParam* InOtherParam)
+	: Data(InOtherParam->Data)
+	, TypeId(InOtherParam->TypeId)
+	, Flags(InOtherParam->Flags)
+{
+}
+
+bool FParam::CanAssignTo(const FParam& InParam) const
+{
+	return CanAssignWith(InParam.GetType(), InParam.Flags);
+}
+
+bool FParam::CanAssignWith(const FParamType& InType, EFlags InFlags, FStringBuilderBase* OutReasonPtr) const
+{
+	// Check type
+	if(TypeId != InType.GetTypeId())
+	{
+		if(OutReasonPtr)
+		{
+			const FParamType& ThisType = FParamType::FRegistrar::GetTypeById(TypeId);
+			OutReasonPtr->Appendf(TEXT("Types do not match: %s and %s"), *ThisType.GetName().ToString(), *InType.GetName().ToString());
+		}
+		return false;
+	}
+	
+	// Check const
+	if(!EnumHasAnyFlags(InFlags, EFlags::Mutable))
+	{
+		if(OutReasonPtr)
+		{
+			OutReasonPtr->Append(TEXT("Target is not mutable"));
+		}
+		return false;
+	}
+
+	// Check chunking
+	if(EnumHasAnyFlags(Flags, EFlags::Chunked) != EnumHasAnyFlags(InFlags, EFlags::Chunked))
+	{
+		if(OutReasonPtr)
+		{
+			OutReasonPtr->Append(TEXT("Incompatible chunking"));
+		}
+		return false;
+	}
+	
+	return true;
+}
+
+void FParamType::FRegistrar::RegisterType(FParamType& InType, const UScriptStruct* InStruct, FName InName, uint32 InSize, uint32 InAlignment)
+{
+	InType.Struct = InStruct;
+	InType.Name = InName;
+	InType.Size = InSize;
+	InType.Alignment = InAlignment;
+
+	check(InType.Name != NAME_None);
+	check(InType.Size != 0);
+	check(InType.Alignment != 0);
+	
+	InType.TypeId = GTypeRegistry.Add(&InType);
+}
+
+const FParamType& FParamType::FRegistrar::FindTypeByName(FName InName)
+{
+	for(FParamType* Type : GTypeRegistry)
+	{
+		if(Type->Name == InName)
+		{
+			return *Type;
+		}
+	}
+
+	checkf(false, TEXT("Type %s not found"), *InName.ToString());
+	return DefaultType;
+}
+
+const FParamType& FParamType::FRegistrar::GetTypeById(uint16 InTypeId)
+{
+	if(GTypeRegistry.IsValidIndex(InTypeId))
+	{
+		return *GTypeRegistry[InTypeId];
+	}
+
+	checkf(false, TEXT("Type ID %d not found"), InTypeId);
+	return DefaultType;	
+}
+
+FParamType::FRegistrar::FRegistrar(TUniqueFunction<void(void)>&& InFunction)
+{
+	check(IsInGameThread());
+	GDeferredTypes.Add(MoveTemp(InFunction));
+}
+
+void FParamType::FRegistrar::RegisterDeferredTypes()
+{
+	check(IsInGameThread());
+	
+	for(TUniqueFunction<void(void)>& Function : GDeferredTypes)
+	{
+		Function();
+	}
+
+	GDeferredTypes.Empty();
+}
+
+const FParamType& FParam::GetType() const
+{
+	check(GTypeRegistry.IsValidIndex(TypeId));
+	return *GTypeRegistry[TypeId];
+}
+/*
+static FParamType& GetTypeForProperty(const FProperty* InProperty)
+{
+	const FName TypeName = *InProperty->GetCPPType();
+	return FParamType::FRegistrar::FindTypeByName(TypeName);
+}
+
+FPropertyParam::FPropertyParam(const FProperty* InProperty, UObject* InContainer)
+	: FParam(GetTypeForProperty(InProperty), InProperty->ContainerPtrToValuePtr<void>(InContainer), EFlags::Const)	// Property parameters are assumed to be const
+{
+}
+*/
+}
diff --git a/Engine/Plugins/Experimental/DataInterface/Source/DataInterface/Private/DataInterfaceState.cpp b/Engine/Plugins/Experimental/DataInterface/Source/DataInterface/Private/DataInterfaceState.cpp
new file mode 100644
--- /dev/null
+++ b/Engine/Plugins/Experimental/DataInterface/Source/DataInterface/Private/DataInterfaceState.cpp
@@ -0,0 +1,58 @@
+﻿// Copyright Epic Games, Inc. All Rights Reserved.
+
+#include "DataInterfaceState.h"
+#include "DataInterfaceContext.h"
+
+namespace UE::DataInterface
+{
+
+FParam* FState::FindStateRaw(const FInterfaceKeyWithIdAndStack& InKey, const FContext& InContext, EStatePersistence InPersistence)
+{
+	check(NumElements == InContext.GetNum())
+
+	FParam* ExistingValue = nullptr;
+	TMap<FInterfaceKeyWithIdAndStack, FParam>* ValueMap = nullptr; 
+
+	switch(InPersistence)
+	{
+	case EStatePersistence::Relevancy:
+		ExistingValue = RelevancyValueMap.Find(InKey);
+		if(ExistingValue)
+		{
+			// An access to a state param counts as a 'relevant use' for this update
+			static_cast<FRelevancyParam*>(ExistingValue)->UpdateCounter = InContext.UpdateCounter;
+		}
+		break;
+	case EStatePersistence::Permanent:
+		ExistingValue = PermanentValueMap.Find(InKey);
+		break;
+	default:
+		check(false);
+		break;
+	}
+	
+	return ExistingValue;
+}
+
+FParam* FState::AllocateState(const FInterfaceKeyWithIdAndStack& InKey, const FContext& InContext, const FParamType& InType, EStatePersistence InPersistence)
+{
+	void* Data = FMemory::Malloc(InType.GetSize() * NumElements, InType.GetAlignment());
+	
+	switch(InPersistence)
+	{
+	case EStatePersistence::Relevancy:
+		// TODO: chunked allocator for relevancy-based stuff?
+		return &RelevancyValueMap.Add(InKey, FRelevancyParam(InType, Data, InContext.GetChunkFlags(), InContext.UpdateCounter));
+		break;
+	case EStatePersistence::Permanent:
+		return &PermanentValueMap.Add(InKey, FParam(InType, Data, InContext.GetChunkFlags()));
+		break;
+	default:
+		check(false);
+		break;
+	}
+	
+	return nullptr;
+}
+
+}
diff --git a/Engine/Plugins/Experimental/DataInterface/Source/DataInterface/Private/DataInterfaceTypes.cpp b/Engine/Plugins/Experimental/DataInterface/Source/DataInterface/Private/DataInterfaceTypes.cpp
new file mode 100644
--- /dev/null
+++ b/Engine/Plugins/Experimental/DataInterface/Source/DataInterface/Private/DataInterfaceTypes.cpp
@@ -0,0 +1,13 @@
+﻿// Copyright Epic Games, Inc. All Rights Reserved.
+
+#include "DataInterfaceTypes.h"
+#include "IDataInterface.h"
+
+namespace UE::DataInterface
+{
+
+#define DATA_INTERFACE_TYPE(Type, Identifier) IMPLEMENT_DATA_INTERFACE_PARAM_TYPE(Type, Identifier)
+#include "DataInterfaceTypes.inl"
+#undef DATA_INTERFACE_TYPE
+
+}
\ No newline at end of file
diff --git a/Engine/Plugins/Experimental/DataInterface/Source/DataInterface/Public/DataInterface.h b/Engine/Plugins/Experimental/DataInterface/Source/DataInterface/Public/DataInterface.h
new file mode 100644
--- /dev/null
+++ b/Engine/Plugins/Experimental/DataInterface/Source/DataInterface/Public/DataInterface.h
@@ -0,0 +1,44 @@
+﻿// Copyright Epic Games, Inc. All Rights Reserved.
+
+#pragma once
+
+#include "CoreMinimal.h"
+#include "DataInterfaceContext.h"
+
+namespace UE::DataInterface
+{
+
+template<typename InterfaceType, typename ValueType>
+static bool GetDataSafe(const TScriptInterface<InterfaceType> DataInterface, const FContext& InContext, ValueType& Result)
+{
+	if(DataInterface.GetInterface() != nullptr)
+	{
+		if constexpr (TIsDerivedFrom<ValueType, FParam>::Value)
+		{
+			return DataInterface.GetInterface()->GetData(InContext, Result);
+		}
+		else
+		{
+			TWrapParam<ValueType> WrappedResult(InContext, &Result);
+			return DataInterface.GetInterface()->GetData(InContext, WrappedResult);
+		}
+	}
+
+	return false;
+}
+
+template<typename ValueType>
+static bool GetParameter(const FContext& InContext, FName InKey, ValueType& Result)
+{
+	if constexpr (TIsDerivedFrom<ValueType, FParam>::Value)
+	{
+		return InContext.GetParameter(InKey, Result);
+	}
+	else
+	{
+		TWrapParam<ValueType> WrappedResult(InContext, &Result);
+		return InContext.GetParameter(InKey, WrappedResult);
+	}
+}
+
+}
\ No newline at end of file
diff --git a/Engine/Plugins/Experimental/DataInterface/Source/DataInterface/Public/DataInterfaceContext.h b/Engine/Plugins/Experimental/DataInterface/Source/DataInterface/Public/DataInterfaceContext.h
new file mode 100644
--- /dev/null
+++ b/Engine/Plugins/Experimental/DataInterface/Source/DataInterface/Public/DataInterfaceContext.h
@@ -0,0 +1,244 @@
+// Copyright Epic Games, Inc. All Rights Reserved.
+
+#pragma once
+
+#include "CoreMinimal.h"
+#include "UObject/Class.h"
+#include "DataInterfaceParam.h"
+#include "DataInterfaceState.h"
+#include "DataInterfaceKey.h"
+#include "Misc/MemStack.h"
+
+class IDataInterface;
+class IDataInterfaceParameters;
+
+namespace UE::DataInterface
+{
+
+struct FState;
+enum class EStatePersistence : uint8;
+
+// Helper to provide a non-callstack bridge for interfacing with data interface contexts
+// Declaring one of these on the stack will enable all calls in its scope to access the passed-in context
+// via FThreadContext::Get()
+struct DATAINTERFACE_API FThreadContext
+{
+	FThreadContext(const FContext& InContext);
+	~FThreadContext();
+
+	static const FContext& Get();
+};
+
+// Context providing methods for mutating & interrogating the data interface runtime
+struct DATAINTERFACE_API FContext
+{
+private:
+	friend class ::IDataInterface;
+	friend struct FState;
+
+public:
+	// Root public constructor. Constructs a context given a state.
+	// Further usage must use WithResult etc. to allow writable results to be used or param
+	FContext(float InDeltaTime, FState& InState, IDataInterfaceParameters* InParameters = nullptr);
+
+	FContext WithResult(FParam& InResult) const;
+
+	FContext WithParameter(FName ParameterId, const FParam& InParameter) const;
+
+	FContext WithParameters(TArrayView<const TPair<FName, FParam>> InParameters) const;
+	
+	FContext WithParameters(IDataInterfaceParameters* InParameters) const;
+	
+	bool GetParameter(FName InKey, FParam& OutParam) const;
+
+public:
+	// Gets (and allocates, if necessary) state for the specified node given this calling context
+	template<typename ValueType, EStatePersistence Persistence = EStatePersistence::Relevancy>
+	TParam<ValueType> GetState(const FInterfaceKeyWithId& InKey) const
+	{
+		return State.GetState<ValueType, Persistence>(InKey, *this, CallstackHash);
+	}
+
+	// Gets (and allocates, if necessary) state for the specified node given this calling context
+	template<typename ValueType, EStatePersistence Persistence = EStatePersistence::Relevancy>
+	TParam<ValueType> GetState(const IDataInterface* InDataInterface, uint32 InId) const
+	{
+		return State.GetState<ValueType, Persistence>(InDataInterface, InId, *this, CallstackHash);
+	}
+	
+public:
+	// Set a result directly
+	template<typename ValueType>
+	void SetResult(const ValueType& InValue) const
+	{
+		if(Result != nullptr)
+		{
+			// Support containers with Num() and GetData()
+			if constexpr (TModels<Private::CSizedContainerWithAccessibleDataAsRawPtr, ValueType>::Value)
+			{
+				checkSlow(Private::CheckParam<ValueType>(Result));
+				checkSlow(InValue.Num() == State.ChunkSize);
+
+				TArrayView<ValueType> Results(static_cast<ValueType*>(Result->Data), State.NumElements);
+				const ValueType* Data = InValue.GetData();
+				for(int32 ResultIndex = 0; ResultIndex < State.NumElements; ++ResultIndex)
+				{
+					Results[ResultIndex] = Data[ResultIndex];
+				}
+			}
+			else
+			{
+				checkSlow(Private::CheckParam<ValueType>(Result));
+
+				TArrayView<ValueType> Results(static_cast<ValueType*>(Result->Data), State.NumElements);
+				for(ValueType& ResultElement : Results)
+				{
+					ResultElement = InValue;
+				}
+			}
+		}
+	}
+	
+	// Get the current result as a specified type
+	template<typename ValueType>
+	TParam<ValueType> GetResult() const
+	{
+		// Cannot call this without a result being present
+		check(Result);
+		
+		// Types must match
+		checkSlow(Private::CheckParam<ValueType>(*Result));
+		return TParam<ValueType>(*Result, GetChunkFlags());
+	}
+	
+	// Get the current result as a raw parameter
+	const FParam& GetResultRaw() const
+	{
+		// Cannot call this without a result being present
+		check(Result);
+		return *Result;
+	}
+
+public:
+	// Number of elements that this context represents
+	uint32 GetNum() const
+	{
+		return State.NumElements;
+	}
+
+	// Max number of elements that this context represents
+	uint32 GetChunkSize() const
+	{
+		return State.ChunkSize;
+	}
+
+	// Access delta time as a param
+	TParam<const float> GetDeltaTimeParam() const;
+
+	// Raw access to delta time
+	float GetDeltaTime() const { return DeltaTime; }
+	
+	// Get parameter flags for chunking
+	FParam::EFlags GetChunkFlags() const
+	{
+		return State.ChunkSize > 1 ? FParam::EFlags::Chunked : FParam::EFlags::None;
+	}
+
+private:
+	FContext(float InDeltaTime, FState& InState, FParam* InResult)
+		: State(InState)
+		, Result(InResult)
+		, DeltaTime(InDeltaTime)
+	{}
+
+	FContext WithCallRaw(const IDataInterface* InDataInterface) const;
+
+	void FlushRelevancy() const;
+
+private:
+	TMap<FName, FParam> AdditionalParameters;
+	const FContext* Parent = nullptr;
+	const FContext* Root = nullptr;
+	FState& State;
+	FParam* Result = nullptr;
+	IDataInterfaceParameters* Parameters = nullptr;
+	float DeltaTime = 0.0f;
+	uint32 CallstackHash = 0;
+	uint32 UpdateCounter = 0;
+};
+
+// A typed param that wraps a user ptr
+template<typename ValueType>
+struct TWrapParam : TParam<ValueType>
+{
+public:
+	TWrapParam(const FContext& InContext, ValueType* InValuePtrToWrap)
+		: TParam<ValueType>(GetDataForConstructor(InContext, InValuePtrToWrap), GetFlagsForConstructor(InContext, InValuePtrToWrap))
+	{
+	}
+
+private:
+	void* GetDataForConstructor(const FContext& InContext, ValueType* InValuePtrToWrap) const
+	{
+		using MutableValueType = typename TRemoveConst<ValueType>::Type;
+		
+		// Support containers with Num() and GetData()
+		if constexpr (TModels<Private::CSizedContainerWithAccessibleDataAsRawPtr, ValueType>::Value)
+		{
+			checkSlow(InContext.GetChunkSize() == InValuePtrToWrap->Num());
+			return const_cast<void*>(InValuePtrToWrap->GetData());
+		}
+		else
+		{
+			checkSlow(InContext.GetChunkSize() == 1);
+			return const_cast<MutableValueType*>(InValuePtrToWrap);
+		}
+	}
+
+	FParam::EFlags GetFlagsForConstructor(const FContext& InContext, ValueType* InValuePtrToWrap) const
+	{
+		FParam::EFlags Flags = FParam::EFlags::None;
+		
+		// Add chunked flags or not
+		if constexpr (TModels<Private::CSizedContainerWithAccessibleDataAsRawPtr, ValueType>::Value)
+		{
+			// Containers always assume chunked
+			checkSlow(InContext.GetChunkSize() == InValuePtrToWrap->Num());
+			Flags |= InContext.GetChunkFlags();
+		}
+		else
+		{
+			checkSlow(InContext.GetChunkSize() == 1);
+		}
+
+		// Add const flags or not
+		if constexpr (!TIsConst<ValueType>::Value)
+		{
+			Flags |= FParam::EFlags::Mutable;
+		}
+
+		return Flags;
+	}
+};
+
+// A typed param that owns it's own memory
+template<typename ValueType, typename AllocatorType = TMemStackAllocator<>>
+struct TAllocParam : TParam<ValueType>
+{
+public:
+	TAllocParam(const FContext& InContext)
+		: TParam<ValueType>(GetDataForConstructor(InContext), InContext.GetChunkFlags())
+	{
+	}
+
+private:
+	void* GetDataForConstructor(const FContext& InContext)
+	{
+		ValueArray.SetNum(InContext.GetChunkSize());
+		return ValueArray.GetData();
+	}
+
+	TArray<ValueType, AllocatorType> ValueArray;
+};
+
+}
\ No newline at end of file
diff --git a/Engine/Plugins/Experimental/DataInterface/Source/DataInterface/Public/DataInterfaceKernel.h b/Engine/Plugins/Experimental/DataInterface/Source/DataInterface/Public/DataInterfaceKernel.h
new file mode 100644
--- /dev/null
+++ b/Engine/Plugins/Experimental/DataInterface/Source/DataInterface/Public/DataInterfaceKernel.h
@@ -0,0 +1,71 @@
+﻿// Copyright Epic Games, Inc. All Rights Reserved.
+
+#pragma once
+
+#include "CoreMinimal.h"
+#include "DataInterfaceContext.h"
+
+namespace UE::DataInterface
+{
+
+struct DATAINTERFACE_API FKernel
+{
+public:
+	// Runs a 'kernel' predicate that iterates over the supplied (optionally chunked) params
+	// Number and type of params must match between the predicate & param pack
+	template<typename PredicateType, typename... ParamTypes>
+	static void Run(const FContext& InContext, PredicateType InPredicate, ParamTypes&... InArgs)
+	{
+		// Cache offsets to apply each element for each parameter
+		const uint32 Offsets[sizeof...(InArgs)] =
+		{
+			[](auto& InArg) -> uint32
+			{
+				if constexpr (TIsDerivedFrom<ParamTypes, FParam>::Value)
+				{
+					if(EnumHasAnyFlags(InArg.Flags, FParam::EFlags::Chunked))
+					{
+						return InArg.GetType().GetSize();
+					}
+				}
+
+				return 0;
+			}(InArgs)...
+		};
+
+		// Iterate per-element in the context, applying kernel predicate to all elements
+		const uint32 NumElements = InContext.GetNum();
+		for(uint32 ElementIndex = 0; ElementIndex < NumElements; ++ElementIndex)
+		{
+			uint32 ParameterIndex = 0;
+
+			// Call predicate, expanding InArgs parameter pack from TParam<ParamType> -> ParamType
+			// and applying correct offset for each param
+			InPredicate(
+				[&Offsets, &ParameterIndex](uint32 InElementIndex, auto& InArg) -> auto&
+				{
+					auto GetBasePtr = [](auto& InArg)
+					{
+						if constexpr (TIsDerivedFrom<ParamTypes, FParam>::Value)
+						{
+							return static_cast<typename ParamTypes::TValueType*>(InArg.Data);
+						}
+						else
+						{
+							return static_cast<ParamTypes*>(&InArg);
+						}
+					};
+
+					auto GetOffset = [&Offsets, &ParameterIndex](uint32 InElementIndex) -> uint32
+					{
+						return Offsets[ParameterIndex++] * InElementIndex;
+					};
+		
+					return *(GetBasePtr(InArg) + GetOffset(InElementIndex));
+				}(ElementIndex, InArgs)...
+			);
+		}
+	}
+};
+
+}
\ No newline at end of file
diff --git a/Engine/Plugins/Experimental/DataInterface/Source/DataInterface/Public/DataInterfaceKey.h b/Engine/Plugins/Experimental/DataInterface/Source/DataInterface/Public/DataInterfaceKey.h
new file mode 100644
--- /dev/null
+++ b/Engine/Plugins/Experimental/DataInterface/Source/DataInterface/Public/DataInterfaceKey.h
@@ -0,0 +1,68 @@
+﻿// Copyright Epic Games, Inc. All Rights Reserved.
+
+#pragma once
+
+#include "CoreMinimal.h"
+#include "UObject/Class.h"
+#include "DataInterfaceParam.h"
+
+class IDataInterface;
+
+namespace UE::DataInterface
+{
+
+// Key value for cached data interface state
+struct DATAINTERFACE_API FInterfaceKey
+{
+	FInterfaceKey(const IDataInterface* InDataInterface)
+		: DataInterface(InDataInterface)
+		, Hash(::GetTypeHash(InDataInterface))
+	{}
+
+	FInterfaceKey(const FInterfaceKey& InExistingKey, uint32 InHash)
+		: DataInterface(InExistingKey.DataInterface)
+		, Hash(HashCombineFast(InExistingKey.Hash, InHash))
+	{}
+
+	const IDataInterface* DataInterface = nullptr;
+	uint32 Hash = 0;
+};
+
+// Key value for cached data interface state
+struct DATAINTERFACE_API FInterfaceKeyWithId : FInterfaceKey
+{
+	FInterfaceKeyWithId(const FInterfaceKey& InInterfaceKey, uint32 InId)
+		: FInterfaceKey(InInterfaceKey, ::GetTypeHash(InId))
+		, Id(InId)
+	{}
+
+	FInterfaceKeyWithId(const FInterfaceKeyWithId& InInterfaceKeyWithId, uint32 InHash)
+		: FInterfaceKey(InInterfaceKeyWithId, HashCombineFast(InInterfaceKeyWithId.Hash, InHash))
+		, Id(InInterfaceKeyWithId.Id)
+	{}
+	
+	uint32 Id = 0;
+};
+
+// Key value for cached data interface state
+struct DATAINTERFACE_API FInterfaceKeyWithIdAndStack : FInterfaceKeyWithId
+{
+	FInterfaceKeyWithIdAndStack(const FInterfaceKeyWithId& InInterfaceKey, uint32 InCallStackHash)
+		: FInterfaceKeyWithId(InInterfaceKey, InCallStackHash)
+		, CallStackHash(InCallStackHash)
+	{}
+
+	FORCEINLINE friend bool operator==(const FInterfaceKeyWithIdAndStack& A, const FInterfaceKeyWithIdAndStack& B)
+	{
+		return A.DataInterface == B.DataInterface && A.Id == B.Id && A.CallStackHash == B.CallStackHash;
+	}
+
+	FORCEINLINE friend uint32 GetTypeHash(const FInterfaceKeyWithIdAndStack& InKey)
+	{
+		return InKey.Hash;
+	}
+
+	uint32 CallStackHash = 0;
+};
+
+}
\ No newline at end of file
diff --git a/Engine/Plugins/Experimental/DataInterface/Source/DataInterface/Public/DataInterfaceParam.h b/Engine/Plugins/Experimental/DataInterface/Source/DataInterface/Public/DataInterfaceParam.h
new file mode 100644
--- /dev/null
+++ b/Engine/Plugins/Experimental/DataInterface/Source/DataInterface/Public/DataInterfaceParam.h
@@ -0,0 +1,252 @@
+// Copyright Epic Games, Inc. All Rights Reserved.
+
+#pragma once
+
+#include "CoreMinimal.h"
+#include "UObject/Class.h"
+#include "UObject/ScriptInterface.h"
+
+namespace UE::DataInterface
+{
+
+struct FParamType;
+
+// Parameter/result/state memory wrapper
+// @TODO can we save memory here with some base/derived system for fundamental types vs struct types?
+struct DATAINTERFACE_API FParam
+{
+public:
+	enum class EFlags : uint8
+	{
+		None	= 0,		// No flags
+		Mutable	= 1,		// Parameter is mutable, so can be mutated at runtime
+		Chunked = 2,		// Parameter is chunked, so represents a number of values rather than a single value. Chunk size is determined by an associated FContext
+	};
+
+	// Get the type of this param 
+	const FParamType& GetType() const;
+
+	// Check whether the supplied param can be written to by this param
+	// Verifies that types are compatible, destination is mutable and chunking matches
+	bool CanAssignTo(const FParam& InParam) const;
+
+	// Helper function for CanAssignTo
+	bool CanAssignWith(const FParamType& InType, EFlags InFlags, FStringBuilderBase* OutReasonPtr = nullptr) const;
+
+	bool bIsMutable() const { return EnumHasAnyFlags(Flags, EFlags::Mutable); }
+
+	bool bIsChunked() const { return EnumHasAnyFlags(Flags, EFlags::Chunked); }
+
+protected:
+	friend struct FContext;
+	friend struct FState;
+	friend struct FKernel;
+
+	FParam() = default;
+	
+	FParam(const FParam* InOtherParam);
+	FParam(const FParamType& InType, void* InData, EFlags InFlags);
+
+	void* Data = nullptr;
+	int16 TypeId = INDEX_NONE;
+	EFlags Flags = EFlags::None;
+};
+
+ENUM_CLASS_FLAGS(FParam::EFlags);
+
+// Simple type representation for data interfaces
+struct DATAINTERFACE_API FParamType
+{
+	FParamType() = default;
+
+	inline const UScriptStruct* GetStruct() const
+	{
+		return Struct;
+	}
+		
+	inline FName GetName() const
+	{
+		return Name;
+	}
+		
+	inline int32 GetSize() const
+	{
+		return Size;
+	}
+		
+	int32 GetAlignment() const
+	{
+		return Alignment;
+	}
+		
+	int16 GetTypeId() const
+	{
+		return TypeId;
+	}
+
+	struct DATAINTERFACE_API FRegistrar
+	{
+		// Helper for registration (via IMPLEMENT_DATA_INTERFACE_PARAM_TYPE).
+		// Registers the type via deferred callback
+		FRegistrar(TUniqueFunction<void(void)>&& InFunction);
+
+		// Helper for registration (via IMPLEMENT_DATA_INTERFACE_PARAM_TYPE)
+		static void RegisterType(FParamType& InType, const UScriptStruct* InStruct, FName InName, uint32 InSize, uint32 InAlignment);
+
+		// Helper for registration
+		static void RegisterDeferredTypes();
+
+		// Find a type by name. If the type was not found, returns the default (invalid) type
+		static const FParamType& FindTypeByName(FName InName);
+
+		// Get a type by ID. If the type is invalid, returns the default (invalid) type
+		static const FParamType& GetTypeById(uint16 InTypeId);
+	};
+
+private:
+	const UScriptStruct* Struct = nullptr;
+	FName Name = NAME_None;
+	uint32 Size = 0;
+	uint16 Alignment = 0;
+	uint16 TypeId = INDEX_NONE;
+};
+
+// A type-erased param that wraps FProperty memory
+struct FPropertyParam : FParam
+{
+	FPropertyParam(const FProperty* InProperty, UObject* InContainer);
+};
+
+namespace Private
+{
+
+template<typename ValueType> 
+struct TParamTypeImpl
+{
+	FORCEINLINE static const FParamType& GetType()
+	{
+		static_assert(sizeof(ValueType) > 0, "Unimplemented type - please include the file DataInterfaceTypes.h or where DECLARE_DATA_INTERFACE_PARAM_TYPE is defined for your type");
+		static FParamType Default;
+		return Default;
+	}
+};
+
+// Regular types need to have GetType function for this to work
+template<typename ValueType> 
+struct TParamType
+{
+	FORCEINLINE static const FParamType& GetType()
+	{
+		return TParamTypeImpl<typename TRemoveConst<ValueType>::Type>::GetType();
+	}
+};
+
+template<typename ValueType>
+static bool CheckParam(const FParam& InParam)
+{
+	return InParam.GetType().GetTypeId() == Private::TParamType<ValueType>::GetType().GetTypeId();
+}
+
+}
+
+namespace Private
+{
+
+// Concept used to decide whether to use container-based construction of wrapped results
+struct CSizedContainerWithAccessibleDataAsRawPtr
+{
+	template<typename ContainerType>
+	auto Requires(ContainerType& Container) -> decltype(
+		Container.Num(),
+		Container.GetData()
+	);
+};
+
+}
+
+// A typed result which wraps the type-erased underlying param
+template<typename ValueType>
+struct TParam : FParam
+{
+protected:
+	using TValueType = ValueType;
+	
+	friend struct FContext;
+	friend struct FState;
+	friend struct FKernel;
+	
+	TParam(void* InData, EFlags InFlags)
+		: FParam(Private::TParamType<ValueType>::GetType(), InData, InFlags)
+	{
+	}
+
+	TParam(const void* InData, EFlags InFlags)
+		: FParam(Private::TParamType<ValueType>::GetType(), const_cast<void*>(InData), InFlags)
+	{
+	}
+
+	TParam(FParam& InParam, EFlags InAdditionalFlags = EFlags::None)
+		: FParam(&InParam)
+	{
+#if DO_CHECK
+		TStringBuilder<64> Error;
+		checkf(InParam.CanAssignWith(Private::TParamType<ValueType>::GetType(), GetTypeFlags(InAdditionalFlags), &Error),
+			TEXT("Cannot assign type: %s"), Error.ToString());
+#endif
+	}
+
+	static EFlags GetTypeFlags(EFlags InAdditionalFlags)
+	{
+		EFlags Flags = InAdditionalFlags;
+
+		// Add const flags or not
+		if constexpr (!TIsConst<ValueType>::Value)
+		{
+			Flags |= FParam::EFlags::Mutable;
+		}
+
+		return Flags;
+	}
+};
+
+namespace Private
+{
+
+struct CHasStaticStruct
+{
+	template<typename ValueType>
+	auto Requires(UScriptStruct* Result) -> decltype(
+		Result = ValueType::StaticStruct()
+	);
+};
+
+template<typename ValueType> 
+struct TParamTypeHelper
+{
+	FORCEINLINE static const UScriptStruct* GetStruct()
+	{
+		if constexpr (TModels<CHasStaticStruct, ValueType>::Value)
+		{
+			return TBaseStructure<ValueType>::Get();
+		}
+		else
+		{
+			return nullptr;
+		}
+	}
+
+	FORCEINLINE static uint32 GetSize()
+	{
+		return sizeof(ValueType);
+	}
+
+	FORCEINLINE static uint32 GetAlignment()
+	{
+		return alignof(ValueType);
+	}
+};
+
+}
+
+}
+
diff --git a/Engine/Plugins/Experimental/DataInterface/Source/DataInterface/Public/DataInterfaceState.h b/Engine/Plugins/Experimental/DataInterface/Source/DataInterface/Public/DataInterfaceState.h
new file mode 100644
--- /dev/null
+++ b/Engine/Plugins/Experimental/DataInterface/Source/DataInterface/Public/DataInterfaceState.h
@@ -0,0 +1,90 @@
+﻿// Copyright Epic Games, Inc. All Rights Reserved.
+
+#pragma once
+
+#include "CoreMinimal.h"
+#include "DataInterfaceParam.h"
+#include "DataInterfaceKey.h"
+
+namespace UE::DataInterface
+{
+
+struct FContext;
+struct FParamType;
+
+enum class EStatePersistence : uint8
+{
+	None,
+	Relevancy,
+	Permanent,
+};
+
+// Container for data interface state
+struct DATAINTERFACE_API FState
+{
+	FState(uint32 InNumElements = 1, uint32 InChunkSize = 1)
+		: NumElements(InNumElements)
+		, ChunkSize(InChunkSize)
+	{
+		check(InNumElements > 0);
+		check(InChunkSize > 0);
+	}
+	
+private:
+	friend struct FContext;
+
+	// Gets (and allocates, if necessary) state for the specified node given this calling context
+	template<typename ValueType, EStatePersistence Persistence = EStatePersistence::Relevancy>
+	TParam<ValueType> GetState(const FInterfaceKeyWithId& InKey, const FContext& InContext, uint32 InHash)
+	{
+		const FInterfaceKeyWithIdAndStack CombinedKey(InKey, InHash);
+
+		if(FParam* ExistingState = FindStateRaw(CombinedKey, InContext, Persistence))
+		{
+			checkSlow(Private::CheckParam<ValueType>(*ExistingState));
+			return TParam<ValueType>(*ExistingState);
+		}
+		else
+		{
+			// Allocate new state
+			FParam* NewState = AllocateState(CombinedKey, InContext, Private::TParamType<ValueType>::GetType(), Persistence);
+
+			// Construct
+			new (NewState->Data) ValueType();
+
+			// Return as param
+			return TParam<ValueType>(*NewState);
+		}
+	}
+
+	// Gets (and allocates, if necessary) state for the specified node given this calling context
+	template<typename ValueType, EStatePersistence Persistence = EStatePersistence::Relevancy>
+	TParam<ValueType> GetState(const IDataInterface* InDataInterface, uint32 InId, const FContext& InContext, uint32 InContextHash)
+	{
+		const FInterfaceKeyWithId Key(InDataInterface, InId);
+		return GetState<ValueType>(Key, InContext, InContextHash);
+	}
+	
+private:
+	FParam* FindStateRaw(const FInterfaceKeyWithIdAndStack& InKey, const FContext& InContext, EStatePersistence InPersistence);
+
+	FParam* AllocateState(const FInterfaceKeyWithIdAndStack& InKey, const FContext& InContext, const FParamType& InType, EStatePersistence InPersistence);
+
+private:
+	struct FRelevancyParam : FParam
+	{
+		FRelevancyParam(const FParamType& InType, void* InData, EFlags InFlags, uint32 InUpdateCounter)
+			: FParam(InType, InData, InFlags)
+			, UpdateCounter(InUpdateCounter)
+		{}
+	
+		uint32 UpdateCounter = 0;
+	};
+	
+	TMap<FInterfaceKeyWithIdAndStack, FRelevancyParam> RelevancyValueMap;
+	TMap<FInterfaceKeyWithIdAndStack, FParam> PermanentValueMap;
+	uint32 NumElements = 0;
+	uint32 ChunkSize = 0;
+};
+
+}
\ No newline at end of file
diff --git a/Engine/Plugins/Experimental/DataInterface/Source/DataInterface/Public/DataInterfaceTypes.h b/Engine/Plugins/Experimental/DataInterface/Source/DataInterface/Public/DataInterfaceTypes.h
new file mode 100644
--- /dev/null
+++ b/Engine/Plugins/Experimental/DataInterface/Source/DataInterface/Public/DataInterfaceTypes.h
@@ -0,0 +1,52 @@
+﻿// Copyright Epic Games, Inc. All Rights Reserved.
+
+#pragma once
+
+#include "CoreMinimal.h"
+#include "DataInterfaceParam.h"
+#include "UObject/ScriptInterface.h"
+
+class IDataInterface;
+
+#define DECLARE_DATA_INTERFACE_PARAM_TYPE_INTERNAL(ValueType, Identifier) \
+	template<> \
+	struct UE::DataInterface::Private::TParamTypeImpl<ValueType> \
+	{ \
+		FORCEINLINE static const UE::DataInterface::FParamType& GetType() \
+		{ \
+			return Identifier##_DataInterfaceTypeInfo; \
+		} \
+	};
+
+#define IMPLEMENT_DATA_INTERFACE_PARAM_TYPE_INTERNAL(ValueType, Identifier) \
+	UE::DataInterface::FParamType::FRegistrar Identifier##_Registrar([](){ UE::DataInterface::FParamType::FRegistrar::RegisterType(Identifier##_DataInterfaceTypeInfo, UE::DataInterface::Private::TParamTypeHelper<ValueType>::GetStruct(), FName(#Identifier), UE::DataInterface::Private::TParamTypeHelper<ValueType>::GetSize(), UE::DataInterface::Private::TParamTypeHelper<ValueType>::GetAlignment()); });
+
+// Types that are passed between data interfaces (i.e. exist in in public headers) should declare via this mechanism
+#define DECLARE_DATA_INTERFACE_PARAM_TYPE(APIType, ValueType, Identifier) \
+	extern APIType UE::DataInterface::FParamType Identifier##_DataInterfaceTypeInfo; \
+	DECLARE_DATA_INTERFACE_PARAM_TYPE_INTERNAL(ValueType, Identifier)
+
+// Types that are passed between data interfaces (i.e. exist in in public headers) should declare via this mechanism
+#define DECLARE_BUILTIN_DATA_INTERFACE_PARAM_TYPE(ValueType, Identifier) \
+	DECLARE_DATA_INTERFACE_PARAM_TYPE(DATAINTERFACE_API, ValueType, Identifier)
+
+// Types that are passed between data interfaces should register via this mechanism
+#define IMPLEMENT_DATA_INTERFACE_PARAM_TYPE(ValueType, Identifier) \
+	UE::DataInterface::FParamType Identifier##_DataInterfaceTypeInfo; \
+	IMPLEMENT_DATA_INTERFACE_PARAM_TYPE_INTERNAL(ValueType, Identifier)
+
+// Structs that are only used for internal state (in a cpp file) can use this simpler registration.
+// If they are exposed in a public header please use the DECLARE_DATA_INTERFACE_PARAM_TYPE & IMPLEMENT_DATA_INTERFACE_PARAM_TYPE macros
+#define IMPLEMENT_DATA_INTERFACE_STATE_TYPE(ValueType, Identifier) \
+	static UE::DataInterface::FParamType Identifier##_DataInterfaceTypeInfo; \
+	static UE::DataInterface::FParamType::FRegistrar Identifier##_Registrar([](){ UE::DataInterface::FParamType::FRegistrar::RegisterType(Identifier##_DataInterfaceTypeInfo, UE::DataInterface::Private::TParamTypeHelper<ValueType>::GetStruct(), FName(#Identifier), UE::DataInterface::Private::TParamTypeHelper<ValueType>::GetSize(), UE::DataInterface::Private::TParamTypeHelper<ValueType>::GetAlignment()); }); \
+	DECLARE_DATA_INTERFACE_PARAM_TYPE_INTERNAL(ValueType, Identifier);
+
+namespace UE::DataInterface
+{
+
+#define DATA_INTERFACE_TYPE(Type, Identifier) DECLARE_BUILTIN_DATA_INTERFACE_PARAM_TYPE(Type, Identifier)
+#include "DataInterfaceTypes.inl"
+#undef DATA_INTERFACE_TYPE
+
+}
\ No newline at end of file
diff --git a/Engine/Plugins/Experimental/DataInterface/Source/DataInterface/Public/DataInterfaceTypes.inl b/Engine/Plugins/Experimental/DataInterface/Source/DataInterface/Public/DataInterfaceTypes.inl
new file mode 100644
--- /dev/null
+++ b/Engine/Plugins/Experimental/DataInterface/Source/DataInterface/Public/DataInterfaceTypes.inl
@@ -0,0 +1,27 @@
+﻿// Copyright Epic Games, Inc. All Rights Reserved.
+
+/** Built in types for data interfaces */
+
+#ifndef DATA_INTERFACE_TYPE
+#error DATA_INTERFACE_TYPE is not defined!
+#endif
+
+DATA_INTERFACE_TYPE(bool, Bool)
+DATA_INTERFACE_TYPE(uint8, Uint8)
+DATA_INTERFACE_TYPE(uint16, Uint16)
+DATA_INTERFACE_TYPE(uint32, Uint32)
+DATA_INTERFACE_TYPE(uint64, Uint64)
+DATA_INTERFACE_TYPE(int8, Int8)
+DATA_INTERFACE_TYPE(int16, Int16)
+DATA_INTERFACE_TYPE(int32, Int32)
+DATA_INTERFACE_TYPE(int64, Int64)
+DATA_INTERFACE_TYPE(float, Float)
+DATA_INTERFACE_TYPE(double, Double)
+DATA_INTERFACE_TYPE(FName, Name)
+DATA_INTERFACE_TYPE(FVector, Vector)
+DATA_INTERFACE_TYPE(FVector4, Vector4)
+DATA_INTERFACE_TYPE(FQuat, Quat)
+DATA_INTERFACE_TYPE(FTransform, Transform)
+DATA_INTERFACE_TYPE(TScriptInterface<IDataInterface>, ScriptInterface_DataInterface)
+
+#undef DATA_INTERFACE_TYPE
\ No newline at end of file
diff --git a/Engine/Plugins/Experimental/DataInterface/Source/DataInterface/Public/IDataInterface.h b/Engine/Plugins/Experimental/DataInterface/Source/DataInterface/Public/IDataInterface.h
new file mode 100644
--- /dev/null
+++ b/Engine/Plugins/Experimental/DataInterface/Source/DataInterface/Public/IDataInterface.h
@@ -0,0 +1,69 @@
+// Copyright Epic Games, Inc. All Rights Reserved.
+
+#pragma once
+
+#include "CoreMinimal.h"
+#include "UObject/Interface.h"
+#include "IDataInterface.generated.h"
+
+class IDataInterface;
+namespace UE::DataInterface
+{
+struct FParam;
+struct FContext;
+}
+
+UINTERFACE(BlueprintType, meta = (CannotImplementInterfaceInBlueprint))
+class DATAINTERFACE_API UDataInterface : public UInterface
+{
+	GENERATED_BODY()
+};
+
+// Empty data interface to support 'any' type
+class DATAINTERFACE_API IDataInterface
+{
+	GENERATED_BODY()
+
+public:
+	/**
+	 * Gets data and stores the value in OutResult.
+	 * @return false if type are incompatible, or if nested calls fail, otherwise true
+	 */
+	bool GetData(const UE::DataInterface::FContext& Context, UE::DataInterface::FParam& OutResult) const;
+
+	/**
+	 * Gets data and stores the value in OutResult, checking whether types are compatible
+	 * @return false if nested calls fail, otherwise true
+	 */
+	bool GetDataChecked(const UE::DataInterface::FContext& Context, UE::DataInterface::FParam& OutResult) const;
+
+	/** Get the name of the return type */
+	FName GetReturnTypeName() const
+	{
+		return GetReturnTypeNameImpl();
+	}
+
+	/** Get the struct of the return type, if any */
+	const UScriptStruct* GetReturnTypeStruct() const
+	{
+		return GetReturnTypeStructImpl();
+	}
+
+private:
+	/** Get data if the types are compatible */
+	bool GetDataIfCompatibleInternal(const UE::DataInterface::FContext& InContext) const;
+	
+	/** Get a value from a data interface with no dynamic or static type checking. Internal use only. */
+	bool GetDataRawInternal(const UE::DataInterface::FContext& InContext) const;
+
+protected:
+	/** Get the return type name, used for dynamic type checking */
+	virtual FName GetReturnTypeNameImpl() const = 0;
+
+	/** Get the return type struct, used for dynamic type checking. If the result is not a struct, this should return nullptr */
+	virtual const UScriptStruct* GetReturnTypeStructImpl() const { return nullptr; }
+	
+	/** Get the value for this interface. @return true if successful, false if unsuccessful. */
+	virtual bool GetDataImpl(const UE::DataInterface::FContext& Context) const = 0;
+};
+
diff --git a/Engine/Plugins/Experimental/DataInterface/Source/DataInterface/Public/IDataInterfaceParameters.h b/Engine/Plugins/Experimental/DataInterface/Source/DataInterface/Public/IDataInterfaceParameters.h
new file mode 100644
--- /dev/null
+++ b/Engine/Plugins/Experimental/DataInterface/Source/DataInterface/Public/IDataInterfaceParameters.h
@@ -0,0 +1,47 @@
+﻿// Copyright Epic Games, Inc. All Rights Reserved.
+
+#pragma once
+
+#include "CoreMinimal.h"
+#include "UObject/Interface.h"
+#include "IDataInterfaceParameters.generated.h"
+
+namespace UE::DataInterface
+{
+struct FParam;
+}
+
+UINTERFACE(BlueprintType, meta = (CannotImplementInterfaceInBlueprint))
+class DATAINTERFACE_API UDataInterfaceParameters : public UInterface
+{
+	GENERATED_BODY()
+};
+
+// Interface representing parameters that can be supplied to a data interface
+class DATAINTERFACE_API IDataInterfaceParameters
+{
+	GENERATED_BODY()
+
+	friend struct UE::DataInterface::FContext;
+	
+	// Get a parameter by name
+	virtual bool GetParameter(FName InKey, UE::DataInterface::FParam& OutParam) = 0;
+
+	// Default utility implementation for UObjects
+	// Treats key as a property name or function
+	static bool GetParameterForObject(UObject* InObject, FName InKey, UE::DataInterface::FParam& OutParam);
+/*	{
+		checkSlow(InObject);
+		const UClass* ObjectClass = InObject->GetClass();
+		if(const FProperty* Property = FindFProperty<FProperty>(ObjectClass, InKey))
+		{
+			OutParam = UE::DataInterface::FPropertyParam(Property, InObject);
+			return true;
+		}
+		else if(const UFunction* Function = FindUField<UFunction>(ObjectClass, InKey))
+		{
+			OutParam = UE::DataInterface::FFunctionParam(Property, InObject);
+		}
+		return false;
+	}*/
+};
diff --git a/Engine/Plugins/Experimental/DataInterfaceGraph/DataInterfaceGraph.uplugin b/Engine/Plugins/Experimental/DataInterfaceGraph/DataInterfaceGraph.uplugin
new file mode 100644
--- /dev/null
+++ b/Engine/Plugins/Experimental/DataInterfaceGraph/DataInterfaceGraph.uplugin
@@ -0,0 +1,45 @@
+{
+	"FileVersion" : 3,
+	"Version" : 1,
+	"VersionName" : "0.1",
+	"FriendlyName" : "Data Interface Graph",
+	"Description" : "Framework for defining functional data flow for animation systems",
+	"Category" : "Animation",
+	"CreatedBy" : "Epic Games, Inc.",
+	"CreatedByURL" : "http://epicgames.com",
+	"DocsURL" : "",
+	"MarketplaceURL" : "",
+	"SupportURL" : "",
+	"EnabledByDefault" : false,
+	"CanContainContent" : false,
+	"IsBetaVersion" : true,
+	"Installed" : false,
+	"Plugins": [
+		{
+			"Name": "DataInterface",
+			"Enabled": true
+		},
+		{
+			"Name": "ControlRig",
+			"Enabled": true
+		}
+	],
+	"Modules" :
+	[
+		{
+			"Name" : "DataInterfaceGraph",
+			"Type" : "Runtime",
+			"LoadingPhase" : "Default"
+		},
+		{
+			"Name" : "DataInterfaceGraphEditor",
+			"Type" : "Editor",
+			"LoadingPhase" : "Default"
+		},
+		{
+			"Name" : "DataInterfaceGraphUncookedOnly",
+			"Type" : "UncookedOnly",
+			"LoadingPhase" : "Default"
+		}
+	]
+}
diff --git a/Engine/Plugins/Experimental/DataInterfaceGraph/Source/DataInterfaceGraph/DataInterfaceGraph.Build.cs b/Engine/Plugins/Experimental/DataInterfaceGraph/Source/DataInterfaceGraph/DataInterfaceGraph.Build.cs
new file mode 100644
--- /dev/null
+++ b/Engine/Plugins/Experimental/DataInterfaceGraph/Source/DataInterfaceGraph/DataInterfaceGraph.Build.cs
@@ -0,0 +1,28 @@
+// Copyright Epic Games, Inc. All Rights Reserved.
+
+namespace UnrealBuildTool.Rules
+{
+	public class DataInterfaceGraph : ModuleRules
+	{
+		public DataInterfaceGraph(ReadOnlyTargetRules Target) : base(Target)
+		{
+		
+			PublicDependencyModuleNames.AddRange(
+				new string[]
+				{
+					"RigVM",
+					"ControlRig",
+				}
+			);
+
+			PrivateDependencyModuleNames.AddRange(
+				new string[]
+				{
+					"Core",
+					"CoreUObject",
+					"DataInterface",
+				}
+			);
+		}
+	}
+}
\ No newline at end of file
diff --git a/Engine/Plugins/Experimental/DataInterfaceGraph/Source/DataInterfaceGraph/Internal/DataInterfaceExecuteContext.h b/Engine/Plugins/Experimental/DataInterfaceGraph/Source/DataInterfaceGraph/Internal/DataInterfaceExecuteContext.h
new file mode 100644
--- /dev/null
+++ b/Engine/Plugins/Experimental/DataInterfaceGraph/Source/DataInterfaceGraph/Internal/DataInterfaceExecuteContext.h
@@ -0,0 +1,52 @@
+﻿// Copyright Epic Games, Inc. All Rights Reserved.
+
+#pragma once
+
+#include "CoreMinimal.h"
+#include "RigVMCore/RigVMExecuteContext.h"
+#include "DataInterfaceExecuteContext.generated.h"
+
+class IDataInterface;
+
+namespace UE::DataInterface
+{
+	struct FContext;
+}
+
+USTRUCT(BlueprintType)
+struct FDataInterfaceExecuteContext : public FRigVMExecuteContext
+{
+	GENERATED_BODY()
+
+	FDataInterfaceExecuteContext() = default;
+
+	const UE::DataInterface::FContext& GetContext() const
+	{
+		check(DataInterfaceContext);
+		return *DataInterfaceContext;
+	}
+
+	void SetResult(bool bInResult) const
+	{
+		check(ResultPtr);
+		*ResultPtr &= bInResult;
+	}
+
+	const IDataInterface* GetInterface() const
+	{
+		check(Interface);
+		return Interface;
+	}
+
+	FORCEINLINE void CopyFrom(const FDataInterfaceExecuteContext& Other)
+	{
+		FRigVMExecuteContext::CopyFrom(Other);
+		DataInterfaceContext = Other.DataInterfaceContext;
+		ResultPtr = Other.ResultPtr;
+	}
+
+private:
+	const UE::DataInterface::FContext* DataInterfaceContext = nullptr;
+	const IDataInterface* Interface = nullptr;
+	bool* ResultPtr = nullptr;
+};
\ No newline at end of file
diff --git a/Engine/Plugins/Experimental/DataInterfaceGraph/Source/DataInterfaceGraph/Internal/DataInterfaceGraph.h b/Engine/Plugins/Experimental/DataInterfaceGraph/Source/DataInterfaceGraph/Internal/DataInterfaceGraph.h
new file mode 100644
--- /dev/null
+++ b/Engine/Plugins/Experimental/DataInterfaceGraph/Source/DataInterfaceGraph/Internal/DataInterfaceGraph.h
@@ -0,0 +1,58 @@
+// Copyright Epic Games, Inc. All Rights Reserved.
+
+#pragma once
+
+#include "CoreMinimal.h"
+#include "IDataInterface.h"
+#include "RigVMCore/RigVM.h"
+#include "DataInterfaceGraph.generated.h"
+
+class UEdGraph;
+class UDataInterfaceGraph;
+namespace UE::DataInterfaceGraphUncookedOnly
+{
+	struct FUtils;
+}
+
+namespace UE::DataInterfaceGraphEditor
+{
+	class FGraphEditor;
+}
+
+namespace UE::DataInterfaceGraph
+{
+	extern DATAINTERFACEGRAPH_API const FName EntryPointName;
+	extern DATAINTERFACEGRAPH_API const FName ResultName;
+}
+
+// A user-created graph of logic used to supply data
+UCLASS()
+class DATAINTERFACEGRAPH_API UDataInterfaceGraph : public UObject, public IDataInterface
+{
+	GENERATED_BODY()
+
+	// IAnimDataInterface interface
+	virtual FName GetReturnTypeNameImpl() const final override;
+	virtual const UScriptStruct* GetReturnTypeStructImpl() const final override;
+	virtual bool GetDataImpl(const UE::DataInterface::FContext& Context) const final override;
+
+	// Support rig VM execution
+	TArray<FRigVMExternalVariable> GetRigVMExternalVariables();
+	
+	friend class UDataInterfaceGraphFactory;
+	friend class UDataInterfaceGraph_EditorData;
+	friend struct UE::DataInterfaceGraphUncookedOnly::FUtils;
+	friend class UE::DataInterfaceGraphEditor::FGraphEditor;
+	friend class UDataInterface_Graph;
+	
+	UPROPERTY()
+	TObjectPtr<URigVM> RigVM;
+
+	UPROPERTY()
+	FRigVMRuntimeSettings VMRuntimeSettings;
+	
+#if WITH_EDITORONLY_DATA
+	UPROPERTY(VisibleAnywhere, Instanced, Category = "Graph", meta = (ShowInnerProperties))
+	TObjectPtr<UObject> EditorData;
+#endif
+};
\ No newline at end of file
diff --git a/Engine/Plugins/Experimental/DataInterfaceGraph/Source/DataInterfaceGraph/Internal/DataInterfaceUnitContext.h b/Engine/Plugins/Experimental/DataInterfaceGraph/Source/DataInterfaceGraph/Internal/DataInterfaceUnitContext.h
new file mode 100644
--- /dev/null
+++ b/Engine/Plugins/Experimental/DataInterfaceGraph/Source/DataInterfaceGraph/Internal/DataInterfaceUnitContext.h
@@ -0,0 +1,20 @@
+// Copyright Epic Games, Inc. All Rights Reserved.
+
+#pragma once
+
+#include "CoreMinimal.h"
+#include "DataInterfaceContext.h"
+#include "Units/RigUnitContext.h"
+
+struct FDataInterfaceUnitContext : public FRigUnitContext
+{
+	FDataInterfaceUnitContext(const IDataInterface* InInterface, const UE::DataInterface::FContext& InDataInterfaceContext, bool& bInResult)
+		: Interface(InInterface)
+		, DataInterfaceContext(InDataInterfaceContext)
+		, bResult(bInResult)
+	{}
+
+	const IDataInterface* Interface = nullptr;
+	const UE::DataInterface::FContext& DataInterfaceContext;
+	bool& bResult;
+};
diff --git a/Engine/Plugins/Experimental/DataInterfaceGraph/Source/DataInterfaceGraph/Internal/RigUnit_DataInterfaceBeginExecution.h b/Engine/Plugins/Experimental/DataInterfaceGraph/Source/DataInterfaceGraph/Internal/RigUnit_DataInterfaceBeginExecution.h
new file mode 100644
--- /dev/null
+++ b/Engine/Plugins/Experimental/DataInterfaceGraph/Source/DataInterfaceGraph/Internal/RigUnit_DataInterfaceBeginExecution.h
@@ -0,0 +1,27 @@
+﻿// Copyright Epic Games, Inc. All Rights Reserved.
+
+#pragma once
+
+#include "Units/RigUnit.h"
+#include "DataInterfaceExecuteContext.h"
+#include "RigUnit_DataInterfaceBeginExecution.generated.h"
+
+/**
+ * Event for driving the skeleton hierarchy with variables and rig elements
+ */
+USTRUCT(meta=(DisplayName="Execute Data Interface", Category="Events", NodeColor="1, 0, 0", Keywords="Begin,Update,Tick,Forward,Event"))
+struct DATAINTERFACEGRAPH_API FRigUnit_DataInterfaceBeginExecution : public FRigUnit
+{
+	GENERATED_BODY()
+
+	RIGVM_METHOD()
+	virtual void Execute(const FRigUnitContext& Context) override;
+
+	virtual FName GetEventName() const override { return EventName; }
+
+	// The execution result
+	UPROPERTY(EditAnywhere, DisplayName = "Execute", Category = "BeginExecution", meta = (Output))
+	FDataInterfaceExecuteContext ExecuteContext;
+
+	static FName EventName;
+};
diff --git a/Engine/Plugins/Experimental/DataInterfaceGraph/Source/DataInterfaceGraph/Internal/RigUnit_DataInterfaceEndExecution.h b/Engine/Plugins/Experimental/DataInterfaceGraph/Source/DataInterfaceGraph/Internal/RigUnit_DataInterfaceEndExecution.h
new file mode 100644
--- /dev/null
+++ b/Engine/Plugins/Experimental/DataInterfaceGraph/Source/DataInterfaceGraph/Internal/RigUnit_DataInterfaceEndExecution.h
@@ -0,0 +1,56 @@
+﻿// Copyright Epic Games, Inc. All Rights Reserved.
+
+#pragma once
+
+#include "Units/RigUnit.h"
+#include "DataInterfaceExecuteContext.h"
+#include "DataInterfaceUnitContext.h"
+#include "DataInterfaceTypes.h"
+#include "RigUnit_DataInterfaceEndExecution.generated.h"
+
+struct FRigUnitContext;
+
+/** Event for writing back calculated results to external variables */
+USTRUCT(meta=(DisplayName="End Execute Data Interface"))
+struct DATAINTERFACEGRAPH_API FRigUnit_DataInterfaceEndExecution : public FRigUnit
+{
+	GENERATED_BODY()
+
+	template<typename ValueType>
+	static void SetResult(const FRigVMExecuteContext& InExecuteContext, const ValueType& InResult)
+	{
+		const FDataInterfaceExecuteContext& DataInterfaceExecuteContext = static_cast<const FDataInterfaceExecuteContext&>(InExecuteContext);
+		DataInterfaceExecuteContext.GetContext().SetResult<ValueType>(InResult);
+	}
+
+public:
+	// The execution result
+	UPROPERTY(EditAnywhere, Transient, DisplayName = "End Execute", Category = "EndExecution", meta = (Input))
+	FDataInterfaceExecuteContext ExecuteContext;
+};
+
+/** Event for writing back a calculated bool */
+USTRUCT(meta=(DisplayName="End Execute Data Interface Bool"))
+struct DATAINTERFACEGRAPH_API FRigUnit_DataInterfaceEndExecution_Bool : public FRigUnit_DataInterfaceEndExecution
+{
+	GENERATED_BODY()
+	
+	RIGVM_METHOD()
+	virtual void Execute(const FRigUnitContext& Context) override;
+	
+	UPROPERTY(EditAnywhere, Category = Result, meta = (Input))
+	bool Result;
+};
+
+/** Event for writing back a calculated float */
+USTRUCT(meta=(DisplayName="End Execute Data Interface Float"))
+struct DATAINTERFACEGRAPH_API FRigUnit_DataInterfaceEndExecution_Float : public FRigUnit_DataInterfaceEndExecution
+{
+	GENERATED_BODY()
+	
+	RIGVM_METHOD()
+	virtual void Execute(const FRigUnitContext& Context) override;
+
+	UPROPERTY(EditAnywhere, Category = Result, meta = (Input))
+	float Result;
+};
\ No newline at end of file
diff --git a/Engine/Plugins/Experimental/DataInterfaceGraph/Source/DataInterfaceGraph/Private/DataInterfaceGraph.cpp b/Engine/Plugins/Experimental/DataInterfaceGraph/Source/DataInterfaceGraph/Private/DataInterfaceGraph.cpp
new file mode 100644
--- /dev/null
+++ b/Engine/Plugins/Experimental/DataInterfaceGraph/Source/DataInterfaceGraph/Private/DataInterfaceGraph.cpp
@@ -0,0 +1,45 @@
+// Copyright Epic Games, Inc. All Rights Reserved.
+
+#include "DataInterfaceGraph.h"
+#include "RigVMCore/RigVMMemoryStorage.h"
+#include "RigUnit_DataInterfaceBeginExecution.h"
+#include "DataInterfaceUnitContext.h"
+
+namespace UE::DataInterfaceGraph
+{
+const FName EntryPointName("GetData");
+const FName ResultName("Result");
+}
+
+FName UDataInterfaceGraph::GetReturnTypeNameImpl() const
+{
+	// TODO: this needs to be user-defined & determined at creation time, using float for now
+	return TNameOf<float>::GetName();
+}
+
+const UScriptStruct* UDataInterfaceGraph::GetReturnTypeStructImpl() const
+{
+	// TODO: this needs to be user-defined & determined at creation time, using nullptr (for float) for now
+	return nullptr;
+}
+
+bool UDataInterfaceGraph::GetDataImpl(const UE::DataInterface::FContext& Context) const
+{
+	bool bResult = true;
+	
+	if(RigVM)
+	{
+		FRigUnitContext RigUnitContext;
+		RigUnitContext.State = EControlRigState::Update;
+		FDataInterfaceUnitContext DataInterfaceUnitContext(this, Context, bResult);
+		void* AdditionalArguments[2] = { &RigUnitContext, &DataInterfaceUnitContext };
+		bResult &= RigVM->Execute(TArray<URigVMMemoryStorage*>(), AdditionalArguments, FRigUnit_DataInterfaceBeginExecution::EventName);
+	}
+	
+	return bResult;
+}
+
+TArray<FRigVMExternalVariable> UDataInterfaceGraph::GetRigVMExternalVariables()
+{
+	return TArray<FRigVMExternalVariable>(); 
+}
\ No newline at end of file
diff --git a/Engine/Plugins/Experimental/DataInterfaceGraph/Source/DataInterfaceGraph/Private/DataInterfaceGraphModule.cpp b/Engine/Plugins/Experimental/DataInterfaceGraph/Source/DataInterfaceGraph/Private/DataInterfaceGraphModule.cpp
new file mode 100644
--- /dev/null
+++ b/Engine/Plugins/Experimental/DataInterfaceGraph/Source/DataInterfaceGraph/Private/DataInterfaceGraphModule.cpp
@@ -0,0 +1,5 @@
+// Copyright Epic Games, Inc. All Rights Reserved.
+
+#include "Modules/ModuleManager.h"
+
+IMPLEMENT_MODULE(FDefaultModuleImpl, DataInterface);
\ No newline at end of file
diff --git a/Engine/Plugins/Experimental/DataInterfaceGraph/Source/DataInterfaceGraph/Private/DataInterface_Float.cpp b/Engine/Plugins/Experimental/DataInterfaceGraph/Source/DataInterfaceGraph/Private/DataInterface_Float.cpp
new file mode 100644
--- /dev/null
+++ b/Engine/Plugins/Experimental/DataInterfaceGraph/Source/DataInterfaceGraph/Private/DataInterface_Float.cpp
@@ -0,0 +1,110 @@
+﻿// Copyright Epic Games, Inc. All Rights Reserved.
+
+#include "DataInterface_Float.h"
+#include "DataInterface.h"
+#include "DataInterfaceTypes.h"
+#include "DataInterfaceKernel.h"
+
+#define LOCTEXT_NAMESPACE "DataInterface"
+
+bool UDataInterface_Float_Multiply::GetDataImpl(const UE::DataInterface::FContext& Context) const
+{
+	using namespace UE::DataInterface;
+	
+	check(Inputs.Num() > 0);
+
+	bool bResult = true;
+	TAllocParam<float> IntermediateParam(Context);
+	TParam<float> Result = Context.GetResult<float>();
+	
+	for(const TScriptInterface<IDataInterface>& Input : Inputs)
+	{
+		bResult &= UE::DataInterface::GetDataSafe(Input, Context, IntermediateParam);
+
+		FKernel::Run(Context,
+			[](float& OutResult, float InIntermediate)
+			{
+				OutResult *= InIntermediate;
+			},
+			Result, IntermediateParam);
+	}
+
+	return bResult;
+}
+
+bool UDataInterface_Float_InterpTo::GetDataImpl(const UE::DataInterface::FContext& Context) const
+{
+	using namespace UE::DataInterface;
+
+	TParam<float> Result = Context.GetResult<float>();
+	const float DeltaTime = Context.GetDeltaTime();
+	
+	TAllocParam<float> CurrentValue(Context);
+	bool bResult = UE::DataInterface::GetDataSafe(Current, Context, CurrentValue);
+
+	TAllocParam<float> TargetValue(Context);
+	bResult &= UE::DataInterface::GetDataSafe(Target, Context, TargetValue);
+
+	TAllocParam<float> SpeedValue(Context);
+	bResult &= UE::DataInterface::GetDataSafe(Speed, Context, SpeedValue);
+
+	FKernel::Run(Context,
+		[DeltaTime](float& OutResult, float InCurrent, float InTarget, float InSpeed)
+		{
+			OutResult = FMath::FInterpConstantTo(InCurrent, InTarget, DeltaTime, InSpeed);
+		},
+		Result, CurrentValue, TargetValue, SpeedValue);
+
+	return bResult;
+}
+
+struct FSpringInterpState
+{
+	float Value;
+	float ValueRate;
+};
+
+IMPLEMENT_DATA_INTERFACE_STATE_TYPE(FSpringInterpState, SpringInterpState);
+
+bool UDataInterface_Float_SpringInterp::GetDataImpl(const UE::DataInterface::FContext& Context) const
+{
+	using namespace UE::DataInterface;
+	
+	const TParam<FSpringInterpState> State = Context.GetState<FSpringInterpState>(this, 0);
+	const float DeltaTime = Context.GetDeltaTime();
+
+	TAllocParam<float> TargetValueParam(Context);
+	bool bResult = UE::DataInterface::GetDataSafe(TargetValue, Context, TargetValueParam);
+
+	TAllocParam<float> TargetValueRateParam(Context);
+	bResult &= UE::DataInterface::GetDataSafe(TargetValueRate, Context, TargetValueRateParam);
+
+	TAllocParam<float> SmoothingTimeParam(Context);
+	bResult &= UE::DataInterface::GetDataSafe(SmoothingTime, Context, SmoothingTimeParam);
+
+	TAllocParam<float> DampingRatioParam(Context);
+	bResult &= UE::DataInterface::GetDataSafe(DampingRatio, Context, DampingRatioParam);
+
+	FKernel::Run(Context,
+	[DeltaTime](FSpringInterpState& InOutState,
+					float InTargetValue,
+					float InTargetValueRate,
+					float InSmoothingTime,
+					float InDampingRatio)
+		{
+			FMath::SpringDamperSmoothing(
+				InOutState.Value,
+				InOutState.ValueRate,
+				InTargetValue,
+				InTargetValueRate,
+				DeltaTime,
+				InSmoothingTime,
+				InDampingRatio);
+		},
+		State, TargetValueParam, TargetValueRateParam, SmoothingTimeParam, DampingRatioParam);
+
+	return bResult;
+}
+
+
+#undef LOCTEXT_NAMESPACE
\ No newline at end of file
diff --git a/Engine/Plugins/Experimental/DataInterfaceGraph/Source/DataInterfaceGraph/Private/DataInterface_Float.h b/Engine/Plugins/Experimental/DataInterfaceGraph/Source/DataInterfaceGraph/Private/DataInterface_Float.h
new file mode 100644
--- /dev/null
+++ b/Engine/Plugins/Experimental/DataInterfaceGraph/Source/DataInterfaceGraph/Private/DataInterface_Float.h
@@ -0,0 +1,103 @@
+// Copyright Epic Games, Inc. All Rights Reserved.
+
+#pragma once
+
+#include "CoreMinimal.h"
+#include "DataInterfaceTypes.h"
+#include "IDataInterface.h"
+#include "DataInterfaceContext.h"
+#include "DataInterface_Float.generated.h"
+
+UCLASS()
+class UDataInterfaceFloat : public UObject, public IDataInterface
+{
+	GENERATED_BODY()
+
+	virtual FName GetReturnTypeNameImpl() const final override
+	{
+		static const FName Name = TNameOf<float>::GetName();
+		return Name;
+	}
+
+	virtual bool GetDataImpl(const UE::DataInterface::FContext& Context) const override PURE_VIRTUAL(UDataInterfaceFloat::GetDataImpl, return false;)
+};
+
+UCLASS()
+class UDataInterface_Float_Literal : public UDataInterfaceFloat
+{
+	GENERATED_BODY()
+
+	// IDataInterface interface
+	virtual bool GetDataImpl(const UE::DataInterface::FContext& Context) const final override
+	{
+		Context.SetResult(Value);
+		return true;
+	}
+	
+	UPROPERTY(EditAnywhere, Category = "Parameters")
+	float Value;
+};
+
+UCLASS()
+class UDataInterface_Float_Multiply : public UDataInterfaceFloat
+{
+	GENERATED_BODY()
+
+	// IDataInterface interface
+	virtual bool GetDataImpl(const UE::DataInterface::FContext& Context) const final override;
+
+	UPROPERTY(EditAnywhere, Category = "Parameters")
+	TArray<TScriptInterface<IDataInterface>> Inputs;
+};
+
+UCLASS()
+class UDataInterface_Float_InterpTo : public UDataInterfaceFloat
+{
+	GENERATED_BODY()
+
+	// IDataInterface interface
+	virtual bool GetDataImpl(const UE::DataInterface::FContext& Context) const final override;
+
+	UPROPERTY(EditAnywhere, Category = "Parameters")
+	TScriptInterface<IDataInterface> Current;
+
+	UPROPERTY(EditAnywhere, Category = "Parameters")
+	TScriptInterface<IDataInterface> Target;
+
+	UPROPERTY(EditAnywhere, Category = "Parameters")
+	TScriptInterface<IDataInterface> Speed;
+};
+
+UCLASS()
+class UDataInterface_Float_DeltaTime : public UDataInterfaceFloat
+{
+	GENERATED_BODY()
+
+	// IDataInterface interface
+	virtual bool GetDataImpl(const UE::DataInterface::FContext& Context) const final override
+	{
+		Context.SetResult(Context.GetDeltaTime());
+		return true;
+	}
+};
+
+UCLASS()
+class UDataInterface_Float_SpringInterp : public UDataInterfaceFloat
+{
+	GENERATED_BODY()
+
+	// IDataInterface interface
+	virtual bool GetDataImpl(const UE::DataInterface::FContext& Context) const final override;
+
+	UPROPERTY(EditAnywhere, Category = "Parameters")
+	TScriptInterface<IDataInterface> TargetValue;
+
+	UPROPERTY(EditAnywhere, Category = "Parameters")
+	TScriptInterface<IDataInterface> TargetValueRate;
+	
+	UPROPERTY(EditAnywhere, Category = "Parameters")
+	TScriptInterface<IDataInterface> SmoothingTime;
+	
+	UPROPERTY(EditAnywhere, Category = "Parameters")
+	TScriptInterface<IDataInterface> DampingRatio;
+};
\ No newline at end of file
diff --git a/Engine/Plugins/Experimental/DataInterfaceGraph/Source/DataInterfaceGraph/Private/DataInterface_Name.cpp b/Engine/Plugins/Experimental/DataInterfaceGraph/Source/DataInterfaceGraph/Private/DataInterface_Name.cpp
new file mode 100644
--- /dev/null
+++ b/Engine/Plugins/Experimental/DataInterfaceGraph/Source/DataInterfaceGraph/Private/DataInterface_Name.cpp
@@ -0,0 +1,4 @@
+﻿// Copyright Epic Games, Inc. All Rights Reserved.
+
+#include "DataInterface_Name.h"
+
diff --git a/Engine/Plugins/Experimental/DataInterfaceGraph/Source/DataInterfaceGraph/Private/DataInterface_Name.h b/Engine/Plugins/Experimental/DataInterfaceGraph/Source/DataInterfaceGraph/Private/DataInterface_Name.h
new file mode 100644
--- /dev/null
+++ b/Engine/Plugins/Experimental/DataInterfaceGraph/Source/DataInterfaceGraph/Private/DataInterface_Name.h
@@ -0,0 +1,28 @@
+// Copyright Epic Games, Inc. All Rights Reserved.
+
+#pragma once
+
+#include "CoreMinimal.h"
+#include "IDataInterface.h"
+#include "DataInterfaceContext.h"
+#include "DataInterface_Name.generated.h"
+
+UCLASS()
+class UDataInterface_Name_Literal : public UObject, public IDataInterface
+{
+	GENERATED_BODY()
+
+	virtual FName GetReturnTypeNameImpl() const final override
+	{
+		return NAME_Name;
+	}
+
+	virtual bool GetDataImpl(const UE::DataInterface::FContext& Context) const final override
+	{
+		Context.SetResult(Value);
+		return true;
+	}
+	
+	UPROPERTY(EditAnywhere, Category = "Parameters")
+	FName Value;
+};
diff --git a/Engine/Plugins/Experimental/DataInterfaceGraph/Source/DataInterfaceGraph/Private/DataInterface_Wrapper.cpp b/Engine/Plugins/Experimental/DataInterfaceGraph/Source/DataInterfaceGraph/Private/DataInterface_Wrapper.cpp
new file mode 100644
--- /dev/null
+++ b/Engine/Plugins/Experimental/DataInterfaceGraph/Source/DataInterfaceGraph/Private/DataInterface_Wrapper.cpp
@@ -0,0 +1,23 @@
+﻿// Copyright Epic Games, Inc. All Rights Reserved.
+
+#include "DataInterface_Wrapper.h"
+
+FName UDataInterface_Wrapper::GetReturnTypeNameImpl() const
+{
+	// We mimic the return type of our output
+	return Output.GetInterface() ? Output.GetInterface()->GetReturnTypeName() : NAME_None;
+}
+
+const UScriptStruct* UDataInterface_Wrapper::GetReturnTypeStructImpl() const
+{
+	// We mimic the return type of our output
+	return Output.GetInterface() ? Output.GetInterface()->GetReturnTypeStruct() : nullptr;
+}
+
+bool UDataInterface_Wrapper::GetDataImpl(const UE::DataInterface::FContext& Context) const
+{
+	// Parameterize output with inputs
+	//const UE::DataInterface::FContext ParameterizedContext = Context.WithParameters(Inputs);
+	//return IDataInterface::StaticGetDataRaw(Output, GetReturnTypeName(), ParameterizedContext, Result);
+	return false;
+}
diff --git a/Engine/Plugins/Experimental/DataInterfaceGraph/Source/DataInterfaceGraph/Private/DataInterface_Wrapper.h b/Engine/Plugins/Experimental/DataInterfaceGraph/Source/DataInterfaceGraph/Private/DataInterface_Wrapper.h
new file mode 100644
--- /dev/null
+++ b/Engine/Plugins/Experimental/DataInterfaceGraph/Source/DataInterfaceGraph/Private/DataInterface_Wrapper.h
@@ -0,0 +1,25 @@
+// Copyright Epic Games, Inc. All Rights Reserved.
+
+#pragma once
+
+#include "CoreMinimal.h"
+#include "IDataInterface.h"
+#include "DataInterface_Wrapper.generated.h"
+
+UCLASS()
+class UDataInterface_Wrapper : public UObject, public IDataInterface
+{
+	GENERATED_BODY()
+
+private:
+	// IAnimDataInterface interface
+	virtual FName GetReturnTypeNameImpl() const final override;
+	virtual const UScriptStruct* GetReturnTypeStructImpl() const final override;
+	virtual bool GetDataImpl(const UE::DataInterface::FContext& Context) const final override;
+
+//	UPROPERTY(EditAnywhere, Category = "Parameters")
+//	TArray<FDataInterfaceParameter> Inputs;
+
+	UPROPERTY(EditAnywhere, Category = "Parameters")
+	TScriptInterface<IDataInterface> Output;
+};
\ No newline at end of file
diff --git a/Engine/Plugins/Experimental/DataInterfaceGraph/Source/DataInterfaceGraph/Private/RigUnit_DataInterfaceBeginExecution.cpp b/Engine/Plugins/Experimental/DataInterfaceGraph/Source/DataInterfaceGraph/Private/RigUnit_DataInterfaceBeginExecution.cpp
new file mode 100644
--- /dev/null
+++ b/Engine/Plugins/Experimental/DataInterfaceGraph/Source/DataInterfaceGraph/Private/RigUnit_DataInterfaceBeginExecution.cpp
@@ -0,0 +1,12 @@
+﻿// Copyright Epic Games, Inc. All Rights Reserved.
+
+#include "RigUnit_DataInterfaceBeginExecution.h"
+#include "Units/RigUnitContext.h"
+
+FName FRigUnit_DataInterfaceBeginExecution::EventName = TEXT("Update");
+
+FRigUnit_DataInterfaceBeginExecution_Execute()
+{
+	const FDataInterfaceExecuteContext& DataInterfaceExecuteContext = static_cast<const FDataInterfaceExecuteContext&>(RigVMExecuteContext);
+	ExecuteContext.CopyFrom(DataInterfaceExecuteContext);
+}
diff --git a/Engine/Plugins/Experimental/DataInterfaceGraph/Source/DataInterfaceGraph/Private/RigUnit_DataInterfaceEndExecution.cpp b/Engine/Plugins/Experimental/DataInterfaceGraph/Source/DataInterfaceGraph/Private/RigUnit_DataInterfaceEndExecution.cpp
new file mode 100644
--- /dev/null
+++ b/Engine/Plugins/Experimental/DataInterfaceGraph/Source/DataInterfaceGraph/Private/RigUnit_DataInterfaceEndExecution.cpp
@@ -0,0 +1,20 @@
+﻿// Copyright Epic Games, Inc. All Rights Reserved.
+
+#include "RigUnit_DataInterfaceEndExecution.h"
+#include "Units/RigUnitContext.h"
+
+FRigUnit_DataInterfaceEndExecution_Bool_Execute()
+{
+	if(Context.State == EControlRigState::Update)
+	{
+		SetResult(RigVMExecuteContext, Result);
+	}
+}
+
+FRigUnit_DataInterfaceEndExecution_Float_Execute()
+{
+	if(Context.State == EControlRigState::Update)
+	{
+		SetResult(RigVMExecuteContext, Result);
+	}
+}
diff --git a/Engine/Plugins/Experimental/DataInterfaceGraph/Source/DataInterfaceGraph/Private/RigUnit_DataInterfaceParameter.cpp b/Engine/Plugins/Experimental/DataInterfaceGraph/Source/DataInterfaceGraph/Private/RigUnit_DataInterfaceParameter.cpp
new file mode 100644
--- /dev/null
+++ b/Engine/Plugins/Experimental/DataInterfaceGraph/Source/DataInterfaceGraph/Private/RigUnit_DataInterfaceParameter.cpp
@@ -0,0 +1,120 @@
+// Copyright Epic Games, Inc. All Rights Reserved.
+
+#include "RigUnit_DataInterfaceParameter.h"
+#include "DataInterfaceUnitContext.h"
+#include "DataInterfaceExecuteContext.h"
+#include "DataInterface.h"
+#include "DataInterfaceTypes.h"
+#include "DataInterfaceKernel.h"
+
+bool FRigUnit_DataInterfaceParameter::GetParameterInternal(FName InName, const FDataInterfaceUnitContext& InContext, void* OutResult)
+{
+	return true;
+}
+
+FRigUnit_DataInterfaceParameter_Float_Execute()
+{
+	if(Context.State == EControlRigState::Update)
+	{
+		GetParameterInternal(Parameter, static_cast<const FDataInterfaceUnitContext&>(Context), &Result);
+	}
+}
+
+FRigUnit_DataInterfaceParameter_DataInterface_Execute()
+{
+	if(Context.State == EControlRigState::Update)
+	{
+		GetParameterInternal(Parameter, static_cast<const FDataInterfaceUnitContext&>(Context), &Result);
+	}
+}
+
+FRigUnit_DataInterface_Float_Execute()
+{
+	using namespace UE::DataInterface;
+
+	// @TODO: ensure that context arg is always present to avoid these checks here
+	// Could bind the RigVM to a specific execution context UStruct?
+	if(Context.State == EControlRigState::Update)
+	{
+		const FDataInterfaceExecuteContext& DataInterfaceExecuteContext = static_cast<const FDataInterfaceExecuteContext&>(RigVMExecuteContext);
+		const FContext& DataInterfaceContext = DataInterfaceExecuteContext.GetContext();
+
+		// Wrap the internal result we are going to be writing to 
+		TWrapParam<float> CallResult(DataInterfaceContext, &Result);
+
+		// Call the interface
+		DataInterfaceExecuteContext.SetResult(UE::DataInterface::GetDataSafe(DataInterface, DataInterfaceContext, CallResult));
+	}
+}
+/*
+FRigUnit_DataInterface_Pose_Execute()
+{
+	using namespace UE::DataInterface;
+
+	// @TODO: ensure that context arg is always present to avoid these checks here
+	// Could bind the RigVM to a specific execution context UStruct?
+	if(Context.State == EControlRigState::Update && RigVMExecuteContext.OpaqueArguments.Num() > 1 && RigVMExecuteContext.OpaqueArguments[1] != nullptr)
+	{
+		const FDataInterfaceUnitContext& DataInterfaceUnitContext = *static_cast<const FDataInterfaceUnitContext*>(RigVMExecuteContext.OpaqueArguments[1]);
+		const FContext& DataInterfaceContext = DataInterfaceUnitContext.DataInterfaceContext;
+
+		// Wrap the internal result we are going to be writing to 
+		TWrapParam<FDataInterfaceExecuteContext> CallResult(DataInterfaceContext, &Result);
+
+		// Call the interface
+		DataInterfaceUnitContext.bResult &= UE::DataInterface::GetDataSafe(PoseInterface, DataInterfaceContext, CallResult);
+	}
+}
+*/
+FRigUnit_FloatOperator_Execute()
+{}
+
+FRigUnit_PoseOperator_Execute()
+{}
+
+FRigUnit_DataInterface_SequencePlayer_Execute()
+{}
+
+struct FSpringDamperState
+{
+	float Value = 0.0f;
+	float ValueRate = 0.0f;
+};
+
+IMPLEMENT_DATA_INTERFACE_STATE_TYPE(FSpringDamperState, SpringDamperState);
+
+FRigUnit_TestFloatState_Execute()
+{
+	using namespace UE::DataInterface;
+
+	// @TODO: ensure that context arg is always present to avoid these checks here
+	// Could bind the RigVM to a specific execution context UStruct?
+	if(Context.State == EControlRigState::Update)
+	{
+		const FDataInterfaceExecuteContext& DataInterfaceExecuteContext = static_cast<const FDataInterfaceExecuteContext&>(RigVMExecuteContext);
+		const FContext& DataInterfaceContext = DataInterfaceExecuteContext.GetContext();
+
+		const TParam<FSpringDamperState> State = DataInterfaceContext.GetState<FSpringDamperState>(DataInterfaceExecuteContext.GetInterface(), 0);
+		const float DeltaTime = DataInterfaceContext.GetDeltaTime();
+
+		FKernel::Run(DataInterfaceContext,
+			[DeltaTime, &Result](FSpringDamperState& InOutState,
+						float InTargetValue,
+						float InTargetValueRate,
+						float InSmoothingTime,
+						float InDampingRatio)
+				{
+					FMath::SpringDamperSmoothing(
+						InOutState.Value,
+						InOutState.ValueRate,
+						InTargetValue,
+						InTargetValueRate,
+						DeltaTime,
+						InSmoothingTime,
+						InDampingRatio);
+
+					Result = InOutState.Value;
+				},
+				State, TargetValue, TargetValueRate, SmoothingTime, DampingRatio);
+	}
+}
diff --git a/Engine/Plugins/Experimental/DataInterfaceGraph/Source/DataInterfaceGraph/Private/RigUnit_DataInterfaceParameter.h b/Engine/Plugins/Experimental/DataInterfaceGraph/Source/DataInterfaceGraph/Private/RigUnit_DataInterfaceParameter.h
new file mode 100644
--- /dev/null
+++ b/Engine/Plugins/Experimental/DataInterfaceGraph/Source/DataInterfaceGraph/Private/RigUnit_DataInterfaceParameter.h
@@ -0,0 +1,273 @@
+// Copyright Epic Games, Inc. All Rights Reserved.
+
+#pragma once
+
+#include "IDataInterface.h"
+#include "Units/RigUnit.h"
+#include "RigUnit_DataInterfaceParameter.generated.h"
+
+struct FDataInterfaceUnitContext;
+class UDataInterfaceGraph;
+
+/** Unit for reading parameters from context */
+USTRUCT(meta=(DisplayName="Parameter"))
+struct FRigUnit_DataInterfaceParameter : public FRigUnit
+{
+	GENERATED_BODY()
+
+protected:
+	static bool GetParameterInternal(FName InName, const FDataInterfaceUnitContext& InContext, void* OutResult);
+
+	UPROPERTY(EditAnywhere, Category = "Parameter", meta = (Input))
+	FName Parameter = NAME_None;
+};
+
+/** Unit for reading float parameter from context */
+USTRUCT(meta=(DisplayName="Float Parameter"))
+struct FRigUnit_DataInterfaceParameter_Float : public FRigUnit_DataInterfaceParameter
+{
+	GENERATED_BODY()
+
+	RIGVM_METHOD()
+	virtual void Execute(const FRigUnitContext& Context) override;
+	
+protected:
+	UPROPERTY(EditAnywhere, Category = "Result", meta = (Output))
+	float Result = 0.0f;
+};
+
+/** Unit for reading data interface graph parameter from context */
+USTRUCT(meta=(DisplayName="Data Interface Parameter"))
+struct FRigUnit_DataInterfaceParameter_DataInterface : public FRigUnit_DataInterfaceParameter
+{
+	GENERATED_BODY()
+
+	RIGVM_METHOD()
+	virtual void Execute(const FRigUnitContext& Context) override;
+	
+protected:
+	UPROPERTY(EditAnywhere, Category = "Result", meta = (Output))
+	TScriptInterface<IDataInterface> Result = nullptr;
+};
+
+USTRUCT(BlueprintType)
+struct FDataInterfaceParameter
+{
+	GENERATED_BODY()
+
+	UPROPERTY(EditAnywhere, Category = "Data Interface")
+	FName Name;
+};
+
+USTRUCT(BlueprintType)
+struct FDataInterfaceParameter_Float
+{
+	GENERATED_BODY()
+
+	UPROPERTY(EditAnywhere, Category = "Data Interface")
+	FName Name;
+
+	UPROPERTY(EditAnywhere, Category = "Data Interface")
+	float Value;
+};
+
+USTRUCT(BlueprintType)
+struct FDataInterfaceParameter_Int
+{
+	GENERATED_BODY()
+
+	UPROPERTY(EditAnywhere, Category = "Data Interface")
+	FName Name;
+
+	UPROPERTY(EditAnywhere, Category = "Data Interface")
+	int32 Value;
+};
+
+USTRUCT(BlueprintType)
+struct FDataInterfaceParameter_Bool
+{
+	GENERATED_BODY()
+
+	UPROPERTY(EditAnywhere, Category = "Data Interface")
+	FName Name;
+
+	UPROPERTY(EditAnywhere, Category = "Data Interface")
+	bool bValue;
+};
+
+USTRUCT(BlueprintType)
+struct FDataInterfaceParameters1
+{
+	GENERATED_BODY()
+
+	UPROPERTY(EditAnywhere, Category = "Data Interface")
+	FDataInterfaceParameter_Int Param0;
+
+	UPROPERTY(EditAnywhere, Category = "Data Interface")
+	float FloatParam;
+};
+
+USTRUCT(BlueprintType)
+struct FDataInterfaceParameter_DataInterface : public FDataInterfaceParameter
+{
+	GENERATED_BODY()
+
+	UPROPERTY(EditAnywhere, Category = "Data Interface")
+	TScriptInterface<IDataInterface> Value = nullptr;
+};
+
+/** Base unit for calling data interfaces from graphs */
+USTRUCT(meta=(DisplayName="Data Interface"))
+struct FRigUnit_DataInterface : public FRigUnit
+{
+	GENERATED_BODY()
+
+protected:
+	UPROPERTY(EditAnywhere, Category = "Data Interface", meta = (Input))
+	TScriptInterface<IDataInterface> DataInterface = nullptr;
+};
+
+/** Unit for getting a float via a data interface */
+USTRUCT(meta=(DisplayName="Data Interface Float"))
+struct FRigUnit_DataInterface_Float : public FRigUnit_DataInterface
+{
+	GENERATED_BODY()
+
+	RIGVM_METHOD()
+	virtual void Execute(const FRigUnitContext& Context) override;
+	
+protected:
+	UPROPERTY(EditAnywhere, Category = "Data Interface", meta = (Input))
+	TArray<FDataInterfaceParameter_DataInterface> Parameters;
+	
+	UPROPERTY(EditAnywhere, Category = "Result", meta = (Output))
+	float Result = 0.0f;
+};
+
+/** Unit for getting a pose via a data interface */
+/*
+USTRUCT(meta=(DisplayName="Get Pose"))
+struct FRigUnit_DataInterface_Pose : public FRigUnit
+{
+	GENERATED_BODY()
+
+	RIGVM_METHOD()
+	virtual void Execute(const FRigUnitContext& Context) override;
+	
+protected:
+	UPROPERTY(EditAnywhere, Category = "Pose Interface", meta = (Input))
+	TScriptInterface<IDataInterface> PoseInterface = nullptr;
+
+	UPROPERTY(EditAnywhere, Category = "Pose Interface", meta = (Input))
+	TArray<FDataInterfaceParameter_DataInterface> Parameters;
+
+	UPROPERTY(EditAnywhere, Transient, DisplayName = "Result", Category = "BeginExecution", meta = (Output))
+	FDataInterfaceExecuteContext Result;
+};
+*/
+/** Unit for getting a pose via a data interface */
+USTRUCT(meta=(DisplayName="Float Operator"))
+struct FRigUnit_FloatOperator : public FRigUnit
+{
+	GENERATED_BODY()
+
+	RIGVM_METHOD()
+	virtual void Execute(const FRigUnitContext& Context) override;
+	
+protected:
+	UPROPERTY(EditAnywhere, Category = "Pose Interface", meta = (Input))
+	TScriptInterface<IDataInterface> Operator = nullptr;
+
+	UPROPERTY(EditAnywhere, Category = "Pose Interface", meta = (Input))
+	float ParamA;
+
+	UPROPERTY(EditAnywhere, Category = "Pose Interface", meta = (Input))
+	float ParamB;
+
+	UPROPERTY(EditAnywhere, Transient, DisplayName = "Result", Category = "BeginExecution", meta = (Output))
+	float Result;
+};
+
+/** Unit for getting a pose via a data interface */
+USTRUCT(meta=(DisplayName="Pose Operator"))
+struct FRigUnit_PoseOperator : public FRigUnit
+{
+	GENERATED_BODY()
+
+	RIGVM_METHOD()
+	virtual void Execute(const FRigUnitContext& Context) override;
+	
+protected:
+	UPROPERTY(EditAnywhere, Category = "Pose Interface", meta = (Input))
+	TScriptInterface<IDataInterface> Operator = nullptr;
+
+	UPROPERTY(EditAnywhere, Category = "Pose Interface", meta = (Input))
+	TArray<FRigVMExecuteContext> InputPoses;
+
+	UPROPERTY(EditAnywhere, Category = "Pose Interface", meta = (Input))
+	FDataInterfaceParameters1 Parameters;
+
+	UPROPERTY(EditAnywhere, Transient, DisplayName = "Result", Category = "BeginExecution", meta = (Output))
+	FRigVMExecuteContext Result;
+};
+
+USTRUCT(BlueprintType)
+struct FAnimSequenceParameters
+{
+	GENERATED_BODY()
+
+	UPROPERTY(EditAnywhere, Category = "Anim Sequence")
+	float PlayRate = 1.0f;
+
+	UPROPERTY(EditAnywhere, Category = "Anim Sequence")
+	float StartPosition = 0.0f;
+
+	UPROPERTY(EditAnywhere, Category = "Anim Sequence")
+	bool bLoop = 0.0f;
+};
+
+/** Unit for getting a pose via a data interface */
+USTRUCT(meta=(DisplayName="Anim Sequence"))
+struct FRigUnit_DataInterface_SequencePlayer : public FRigUnit
+{
+	GENERATED_BODY()
+
+	RIGVM_METHOD()
+	virtual void Execute(const FRigUnitContext& Context) override;
+	
+protected:
+	UPROPERTY(EditAnywhere, Category = "Data Interface", meta = (Input))
+	FAnimSequenceParameters Parameters;
+
+	UPROPERTY(EditAnywhere, Category = "Data Interface", meta = (Input))
+	TScriptInterface<IDataInterface> Sequence = nullptr;
+	
+	UPROPERTY(EditAnywhere, Transient, DisplayName = "Result", Category = "BeginExecution", meta = (Output))
+	FRigVMExecuteContext Result;
+};
+
+/** Unit for getting a pose via a data interface */
+USTRUCT(meta=(DisplayName="Test Float State - Spring Damper Smoothing", Varying))
+struct FRigUnit_TestFloatState : public FRigUnit
+{
+	GENERATED_BODY()
+
+	RIGVM_METHOD()
+	virtual void Execute(const FRigUnitContext& Context) override;
+	
+protected:
+	UPROPERTY(EditAnywhere, Category = "Parameters", meta=(Input))
+	float TargetValue = 0.0f;
+
+	UPROPERTY(EditAnywhere, Category = "Parameters", meta=(Input))
+	float TargetValueRate = 1.0f;
+	
+	UPROPERTY(EditAnywhere, Category = "Parameters", meta=(Input))
+	float SmoothingTime = 1.0f;
+	
+	UPROPERTY(EditAnywhere, Category = "Parameters", meta=(Input))
+	float DampingRatio = 1.0f;
+
+	UPROPERTY(EditAnywhere, Category = "Parameters", meta=(Output))
+	float Result = 0.0f;
+};
\ No newline at end of file
diff --git a/Engine/Plugins/Experimental/DataInterfaceGraph/Source/DataInterfaceGraphEditor/DataInterfaceGraphEditor.Build.cs b/Engine/Plugins/Experimental/DataInterfaceGraph/Source/DataInterfaceGraphEditor/DataInterfaceGraphEditor.Build.cs
new file mode 100644
--- /dev/null
+++ b/Engine/Plugins/Experimental/DataInterfaceGraph/Source/DataInterfaceGraphEditor/DataInterfaceGraphEditor.Build.cs
@@ -0,0 +1,34 @@
+// Copyright Epic Games, Inc. All Rights Reserved.
+
+namespace UnrealBuildTool.Rules
+{
+	public class DataInterfaceGraphEditor : ModuleRules
+	{
+		public DataInterfaceGraphEditor(ReadOnlyTargetRules Target) : base(Target)
+		{
+
+			PrivateDependencyModuleNames.AddRange(
+				new string[]
+				{
+					"Core",
+					"CoreUObject",
+					"Engine",
+					"AssetTools",
+					"DataInterface",
+					"DataInterfaceGraph",
+					"DataInterfaceGraphUncookedOnly",
+					"UnrealEd",
+					"SlateCore",
+					"Slate",
+					"PropertyEditor",
+					"RigVM",
+					"RigVMDeveloper",
+					"ControlRig",
+					"ControlRigDeveloper",
+					"ControlRigEditor",
+					"GraphEditor",
+				}
+			);
+		}
+	}
+}
\ No newline at end of file
diff --git a/Engine/Plugins/Experimental/DataInterfaceGraph/Source/DataInterfaceGraphEditor/Private/AssetTypeActions_DataInterfaceGraph.cpp b/Engine/Plugins/Experimental/DataInterfaceGraph/Source/DataInterfaceGraphEditor/Private/AssetTypeActions_DataInterfaceGraph.cpp
new file mode 100644
--- /dev/null
+++ b/Engine/Plugins/Experimental/DataInterfaceGraph/Source/DataInterfaceGraphEditor/Private/AssetTypeActions_DataInterfaceGraph.cpp
@@ -0,0 +1,37 @@
+﻿// Copyright Epic Games, Inc. All Rights Reserved.
+
+#include "AssetTypeActions_DataInterfaceGraph.h"
+#include "DataInterfaceGraphEditor.h"
+#include "DataInterfaceGraph.h"
+
+#define LOCTEXT_NAMESPACE "AssetTypeActions_DataInterfaceGraph" 
+
+namespace UE::DataInterfaceGraphEditor
+{
+
+void FAssetTypeActions_DataInterfaceGraph::OpenAssetEditor(const TArray<UObject*>& InObjects, TSharedPtr<IToolkitHost> EditWithinLevelEditor)
+{
+	const EToolkitMode::Type Mode = EditWithinLevelEditor.IsValid() ? EToolkitMode::WorldCentric : EToolkitMode::Standalone;
+
+	for (UObject* Object : InObjects)
+	{
+		if(UDataInterfaceGraph* DataInterfaceGraph = Cast<UDataInterfaceGraph>(Object))
+		{
+			TSharedRef<FGraphEditor> GraphEditor = MakeShared<FGraphEditor>();
+			GraphEditor->InitEditor(Mode, EditWithinLevelEditor, DataInterfaceGraph);
+		}
+	}
+}
+
+const TArray<FText>& FAssetTypeActions_DataInterfaceGraph::GetSubMenus() const
+{
+	static const TArray<FText> SubMenus
+	{
+		LOCTEXT("DataInterfacesSubMenu", "Data Interfaces")
+	};
+	return SubMenus;
+}
+
+}
+
+#undef LOCTEXT_NAMESPACE
\ No newline at end of file
diff --git a/Engine/Plugins/Experimental/DataInterfaceGraph/Source/DataInterfaceGraphEditor/Private/AssetTypeActions_DataInterfaceGraph.h b/Engine/Plugins/Experimental/DataInterfaceGraph/Source/DataInterfaceGraphEditor/Private/AssetTypeActions_DataInterfaceGraph.h
new file mode 100644
--- /dev/null
+++ b/Engine/Plugins/Experimental/DataInterfaceGraph/Source/DataInterfaceGraphEditor/Private/AssetTypeActions_DataInterfaceGraph.h
@@ -0,0 +1,24 @@
+﻿// Copyright Epic Games, Inc. All Rights Reserved.
+
+#pragma once
+
+#include "AssetTypeActions_Base.h"
+#include "DataInterfaceGraph.h"
+
+namespace UE::DataInterfaceGraphEditor
+{
+
+class FAssetTypeActions_DataInterfaceGraph : public FAssetTypeActions_Base
+{
+public:
+	// IAssetTypeActions Implementation
+	virtual FText GetName() const override { return NSLOCTEXT("AssetTypeActions", "AssetTypeActions_DataInterfaceGraph", "Data Interface Graph"); }
+	virtual FColor GetTypeColor() const override { return FColor(128,128,64); }
+	virtual UClass* GetSupportedClass() const override { return UDataInterfaceGraph::StaticClass(); }
+	virtual void OpenAssetEditor(const TArray<UObject*>& InObjects, TSharedPtr<IToolkitHost> EditWithinLevelEditor = TSharedPtr<IToolkitHost>()) override;
+	virtual bool CanFilter() override { return true; }
+	virtual const TArray<FText>& GetSubMenus() const override;
+	virtual uint32 GetCategories() override { return EAssetTypeCategories::Animation; }
+};
+
+}
diff --git a/Engine/Plugins/Experimental/DataInterfaceGraph/Source/DataInterfaceGraphEditor/Private/DataInterfaceEditorUtils.cpp b/Engine/Plugins/Experimental/DataInterfaceGraph/Source/DataInterfaceGraphEditor/Private/DataInterfaceEditorUtils.cpp
new file mode 100644
--- /dev/null
+++ b/Engine/Plugins/Experimental/DataInterfaceGraph/Source/DataInterfaceGraphEditor/Private/DataInterfaceEditorUtils.cpp
@@ -0,0 +1,122 @@
+// Copyright Epic Games, Inc. All Rights Reserved.
+
+#include "DataInterfaceEditorUtils.h"
+#include "DataInterfaceGraph_EditorData.h"
+#include "Kismet2/Kismet2NameValidators.h"
+
+namespace UE::DataInterfaceGraphEditor
+{
+
+void FUtils::GetAllGraphNames(const UDataInterfaceGraph_EditorData* InEditorData, TSet<FName>& OutNames)
+{
+
+}
+
+static const int32 MaxNameLength = 100;
+
+FName FUtils::ValidateName(const UDataInterfaceGraph_EditorData* InEditorData, const FString& InName)
+{
+	struct FNameValidator : public INameValidatorInterface
+	{
+		FNameValidator(const UDataInterfaceGraph_EditorData* InEditorData)
+			: EditorData(InEditorData)
+		{
+			GetAllGraphNames(EditorData, Names);
+		}
+		
+		virtual EValidatorResult IsValid (const FName& Name, bool bOriginal = false) override
+		{
+			EValidatorResult ValidatorResult = EValidatorResult::AlreadyInUse;
+
+			if(Name == NAME_None)
+			{
+				ValidatorResult = EValidatorResult::EmptyName;
+			}
+			else if(Name.ToString().Len() > MaxNameLength)
+			{
+				ValidatorResult = EValidatorResult::TooLong;
+			}
+			else
+			{
+				// If it is in the names list then it is already in use.
+				if(!Names.Contains(Name))
+				{
+					ValidatorResult = EValidatorResult::Ok;
+
+					// Check for collision with an existing object.
+					if (UObject* ExistingObject = StaticFindObject(/*Class=*/ nullptr, const_cast<UDataInterfaceGraph_EditorData*>(EditorData), *Name.ToString(), true))
+					{
+						ValidatorResult = EValidatorResult::AlreadyInUse;
+					}
+				}
+			}
+			
+			return ValidatorResult;
+		}
+		
+		virtual EValidatorResult IsValid (const FString& Name, bool bOriginal = false) override
+		{
+			// Converting a string that is too large for an FName will cause an assert, so verify the length
+			if(Name.Len() >= NAME_SIZE)
+			{
+				return EValidatorResult::TooLong;
+			}
+			else if (!FName::IsValidXName(Name, UE_BLUEPRINT_INVALID_NAME_CHARACTERS))
+			{
+				return EValidatorResult::ContainsInvalidCharacters;
+			}
+
+			// If not defined in name table, not current graph name
+			return IsValid( FName(*Name) );
+		}
+
+		/** Name set to validate */
+		TSet<FName> Names;
+		/** The editor data to check for validity within */
+		const UDataInterfaceGraph_EditorData* EditorData;
+	};
+	
+	FString Name = InName;
+	if (Name.StartsWith(TEXT("RigUnit_")))
+	{
+		Name.RightChopInline(8, false);
+	}
+
+	FNameValidator NameValidator(InEditorData);
+
+	// Clean up BaseName to not contain any invalid characters, which will mean we can never find a legal name no matter how many numbers we add
+	if (NameValidator.IsValid(Name) == EValidatorResult::ContainsInvalidCharacters)
+	{
+		for (TCHAR& TestChar : Name)
+		{
+			for (TCHAR BadChar : UE_BLUEPRINT_INVALID_NAME_CHARACTERS)
+			{
+				if (TestChar == BadChar)
+				{
+					TestChar = TEXT('_');
+					break;
+				}
+			}
+		}
+	}
+	
+	int32 Count = 0;
+	FString BaseName = Name;
+	while (NameValidator.IsValid(Name) != EValidatorResult::Ok)
+	{
+		// Calculate the number of digits in the number, adding 2 (1 extra to correctly count digits, another to account for the '_' that will be added to the name
+		int32 CountLength = Count > 0 ? (int32)log((double)Count) + 2 : 2;
+
+		// If the length of the final string will be too long, cut off the end so we can fit the number
+		if (CountLength + BaseName.Len() > MaxNameLength)
+		{
+			BaseName.LeftInline(MaxNameLength - CountLength);
+		}
+		Name = FString::Printf(TEXT("%s_%d"), *BaseName, Count);
+		Count++;
+	}
+
+	return *Name;
+}
+
+}
diff --git a/Engine/Plugins/Experimental/DataInterfaceGraph/Source/DataInterfaceGraphEditor/Private/DataInterfaceEditorUtils.h b/Engine/Plugins/Experimental/DataInterfaceGraph/Source/DataInterfaceGraphEditor/Private/DataInterfaceEditorUtils.h
new file mode 100644
--- /dev/null
+++ b/Engine/Plugins/Experimental/DataInterfaceGraph/Source/DataInterfaceGraphEditor/Private/DataInterfaceEditorUtils.h
@@ -0,0 +1,20 @@
+// Copyright Epic Games, Inc. All Rights Reserved.
+
+#pragma once
+
+#include "CoreMinimal.h"
+
+class UDataInterfaceGraph;
+class UDataInterfaceGraph_EditorData;
+
+namespace UE::DataInterfaceGraphEditor
+{
+
+struct FUtils
+{
+	static FName ValidateName(const UDataInterfaceGraph_EditorData* InEditorData, const FString& InName);
+
+	static void GetAllGraphNames(const UDataInterfaceGraph_EditorData* InEditorData, TSet<FName>& OutNames);
+};
+
+}
diff --git a/Engine/Plugins/Experimental/DataInterfaceGraph/Source/DataInterfaceGraphEditor/Private/DataInterfaceGraphEditor.cpp b/Engine/Plugins/Experimental/DataInterfaceGraph/Source/DataInterfaceGraphEditor/Private/DataInterfaceGraphEditor.cpp
new file mode 100644
--- /dev/null
+++ b/Engine/Plugins/Experimental/DataInterfaceGraph/Source/DataInterfaceGraphEditor/Private/DataInterfaceGraphEditor.cpp
@@ -0,0 +1,322 @@
+﻿// Copyright Epic Games, Inc. All Rights Reserved.
+
+#include "DataInterfaceGraphEditor.h"
+#include "DataInterfaceGraphEditorMode.h"
+#include "DataInterfaceGraph.h"
+#include "DataInterfaceGraph_EditorData.h"
+#include "EdGraphNode_Comment.h"
+#include "EdGraphUtilities.h"
+#include "SDataInterfaceGraphEditorActionMenu.h"
+#include "DataInterfaceUncookedOnlyUtils.h"
+#include "Framework/Commands/GenericCommands.h"
+#include "RigVMModel/RigVMController.h"
+#include "RigVMModel/RigVMGraph.h"
+#include "Kismet2/BlueprintEditorUtils.h"
+#include "ScopedTransaction.h"
+
+#define LOCTEXT_NAMESPACE "DataInterfaceGraphEditor"
+
+namespace UE::DataInterfaceGraphEditor
+{
+
+namespace Modes
+{
+	const FName GraphEditor("DataInterfaceEditorMode");
+}
+
+namespace Tabs
+{
+	const FName Details("DetailsTab");
+	const FName Document("Document");
+}
+
+const FName AppIdentifier("DataInterfaceGraphEditor");
+
+void FGraphEditor::InitEditor(const EToolkitMode::Type InMode, const TSharedPtr<IToolkitHost>& InInitToolkitHost, UDataInterfaceGraph* InDataInterfaceGraph)
+{
+	DataInterfaceGraph = InDataInterfaceGraph;
+
+	DataInterfaceGraph_EditorData = UE::DataInterfaceGraphUncookedOnly::FUtils::GetEditorData(DataInterfaceGraph);
+	DataInterfaceGraph_EditorData->Initialize(false);
+	
+	DocumentManager = MakeShared<FDocumentTracker>();
+	DocumentManager->Initialize(SharedThis(this));
+
+	TSharedRef<FGraphEditorSummoner> GraphEditorSummoner = MakeShared<FGraphEditorSummoner>(SharedThis(this));
+	GraphEditorSummoner->OnCreateGraphEditorWidget().BindSP(this, &FGraphEditor::CreateGraphEditorWidget);
+	GraphEditorSummoner->OnGraphEditorFocused().BindSP(this, &FGraphEditor::OnGraphEditorFocused);
+	GraphEditorSummoner->OnGraphEditorBackgrounded().BindSP(this, &FGraphEditor::OnGraphEditorBackgrounded);
+	GraphEditorTabFactoryPtr = GraphEditorSummoner;
+	DocumentManager->RegisterDocumentFactory(GraphEditorSummoner);
+
+	const bool bCreateDefaultStandaloneMenu = true;
+	const bool bCreateDefaultToolbar = true;
+	FAssetEditorToolkit::InitAssetEditor(InMode, InInitToolkitHost, AppIdentifier, FTabManager::FLayout::NullLayout, bCreateDefaultStandaloneMenu, bCreateDefaultToolbar, InDataInterfaceGraph);
+
+	BindCommands();
+
+	AddApplicationMode(Modes::GraphEditor, MakeShared<FGraphEditorMode>(SharedThis(this)));
+	SetCurrentMode(Modes::GraphEditor);
+	
+	ExtendMenu();
+	ExtendToolbar();
+	RegenerateMenusAndToolbars();
+
+	// Open initial document
+	DocumentManager->OpenDocument(FTabPayload_UObject::Make(DataInterfaceGraph_EditorData->EntryPointGraph), FDocumentTracker::EOpenDocumentCause::OpenNewDocument);
+	DocumentManager->OpenDocument(FTabPayload_UObject::Make(DataInterfaceGraph_EditorData->RootGraph), FDocumentTracker::EOpenDocumentCause::OpenNewDocument);
+	DocumentManager->OpenDocument(FTabPayload_UObject::Make(DataInterfaceGraph_EditorData->FunctionLibraryEdGraph), FDocumentTracker::EOpenDocumentCause::OpenNewDocument);
+}
+
+void FGraphEditor::BindCommands()
+{
+}
+
+void FGraphEditor::ExtendMenu()
+{
+	
+}
+
+void FGraphEditor::ExtendToolbar()
+{
+	
+}
+
+void FGraphEditor::RegisterTabSpawners(const TSharedRef<FTabManager>& InTabManager)
+{
+	DocumentManager->SetTabManager(InTabManager);
+
+	FWorkflowCentricApplication::RegisterTabSpawners(InTabManager);
+}
+
+void FGraphEditor::UnregisterTabSpawners(const TSharedRef<FTabManager>& InTabManager)
+{
+	
+}
+
+FName FGraphEditor::GetToolkitFName() const
+{
+	return FName("DataInterfaceGraphEditor");
+}
+
+FText FGraphEditor::GetBaseToolkitName() const
+{
+	return LOCTEXT("AppLabel", "DataInterfaceGraphEditor");
+}
+
+FString FGraphEditor::GetWorldCentricTabPrefix() const
+{
+	return LOCTEXT("WorldCentricTabPrefix", "DataInterfaceGraphEditor ").ToString();
+}
+
+FLinearColor FGraphEditor::GetWorldCentricTabColorScale() const
+{
+	return FLinearColor(0.3f, 0.2f, 0.5f, 0.5f);
+}
+
+void FGraphEditor::InitToolMenuContext(FToolMenuContext& InMenuContext)
+{
+}
+
+void FGraphEditor::OnGraphEditorFocused(TSharedRef<SGraphEditor> InGraphEditor)
+{
+	// Update the graph editor that is currently focused
+	FocusedGraphEdPtr = InGraphEditor;
+}
+
+void FGraphEditor::OnGraphEditorBackgrounded(TSharedRef<SGraphEditor> InGraphEditor)
+{
+	// Update the graph editor that is currently focused
+	FocusedGraphEdPtr = nullptr;
+}
+
+UEdGraph* FGraphEditor::GetFocusedGraph() const
+{
+	if (FocusedGraphEdPtr.IsValid())
+	{
+		if (UEdGraph* Graph = FocusedGraphEdPtr.Pin()->GetCurrentGraph())
+		{
+			return Graph;
+		}
+	}
+	return nullptr;
+}
+
+URigVMGraph* FGraphEditor::GetFocusedVMGraph() const
+{
+	UDataInterfaceGraph_EdGraph* EdGraph = Cast<UDataInterfaceGraph_EdGraph>(GetFocusedGraph());
+	return DataInterfaceGraph_EditorData->GetRigVMGraph(EdGraph);
+}
+
+URigVMController* FGraphEditor::GetFocusedVMController() const
+{
+	return DataInterfaceGraph_EditorData->GetRigVMController(GetFocusedVMGraph());
+}
+
+TSharedRef<SGraphEditor> FGraphEditor::CreateGraphEditorWidget(TSharedRef<FTabInfo> InTabInfo, UEdGraph* InGraph)
+{
+	if(!GraphEditorCommands.IsValid())
+	{
+		GraphEditorCommands = MakeShared<FUICommandList>();
+	
+		GraphEditorCommands->MapAction(FGenericCommands::Get().Delete,
+			FExecuteAction::CreateSP(this, &FGraphEditor::DeleteSelectedNodes),
+			FCanExecuteAction::CreateSP(this, &FGraphEditor::CanDeleteSelectedNodes));
+	}
+	
+	SGraphEditor::FGraphEditorEvents Events;
+	Events.OnCreateActionMenu = SGraphEditor::FOnCreateActionMenu::CreateSP(this, &FGraphEditor::OnCreateGraphActionMenu);
+	
+	return SNew(SGraphEditor)
+		.AdditionalCommands(GraphEditorCommands)
+		.IsEditable(this, &FGraphEditor::IsEditable, InGraph)
+		.GraphToEdit(InGraph)
+		.GraphEvents(Events)
+		.AssetEditorToolkit(AsShared());
+}
+
+FActionMenuContent FGraphEditor::OnCreateGraphActionMenu(UEdGraph* InGraph, const FVector2D& InNodePosition, const TArray<UEdGraphPin*>& InDraggedPins, bool bAutoExpand, SGraphEditor::FActionMenuClosed InOnMenuClosed)
+{
+	TSharedRef<SActionMenu> ActionMenu = SNew(SActionMenu)
+		.AutoExpandActionMenu(bAutoExpand)
+		.Graph(InGraph)
+		.NewNodePosition(InNodePosition)
+		.DraggedFromPins(InDraggedPins)
+		.OnClosedCallback(InOnMenuClosed);
+
+	TSharedPtr<SWidget> FilterTextBox = StaticCastSharedRef<SWidget>(ActionMenu->GetFilterTextBox());
+	return FActionMenuContent(StaticCastSharedRef<SWidget>(ActionMenu), FilterTextBox);
+}
+
+void FGraphEditor::OnNodeTitleCommitted(const FText& NewText, ETextCommit::Type CommitInfo, UEdGraphNode* NodeBeingChanged)
+{
+	if (UEdGraphNode_Comment* CommentBeingChanged = Cast<UEdGraphNode_Comment>(NodeBeingChanged))
+	{
+		GetFocusedVMController()->SetCommentTextByName(CommentBeingChanged->GetFName(), NewText.ToString(), CommentBeingChanged->FontSize, CommentBeingChanged->bCommentBubbleVisible, CommentBeingChanged->bColorCommentBubble, true, true);
+	}
+}
+
+void FGraphEditor::CloseDocumentTab(const UObject* DocumentID)
+{
+	TSharedRef<FTabPayload_UObject> Payload = FTabPayload_UObject::Make(DocumentID);
+	DocumentManager->CloseTab(Payload);
+}
+
+bool FGraphEditor::InEditingMode() const
+{
+	// @TODO: disallow editing when debugging when implemented
+
+	return true;
+}
+
+bool FGraphEditor::IsEditable(UEdGraph* InGraph) const
+{
+	return InGraph && InEditingMode() && InGraph->bEditable;
+}
+
+FGraphPanelSelectionSet FGraphEditor::GetSelectedNodes() const
+{
+	FGraphPanelSelectionSet CurrentSelection;
+	TSharedPtr<SGraphEditor> FocusedGraphEd = FocusedGraphEdPtr.Pin();
+	if (FocusedGraphEd.IsValid())
+	{
+		CurrentSelection = FocusedGraphEd->GetSelectedNodes();
+	}
+	return CurrentSelection;
+}
+
+void FGraphEditor::DeleteSelectedNodes()
+{
+	TSharedPtr<SGraphEditor> FocusedGraphEd = FocusedGraphEdPtr.Pin();
+	if (!FocusedGraphEd.IsValid())
+	{
+		return;
+	}
+
+	const FScopedTransaction Transaction(FGenericCommands::Get().Delete->GetDescription());
+	FocusedGraphEd->GetCurrentGraph()->Modify();
+	
+	const FGraphPanelSelectionSet SelectedNodes = GetSelectedNodes();
+
+//	SetUISelectionState(NAME_None);
+
+	if(FocusedGraphEd)
+	{
+		FocusedGraphEd->ClearSelectionSet();
+	}
+
+	// Some nodes have sub-objects that are represented as other tabs.
+	// Close them here as a pre-pass before we remove their nodes. If the documents are left open they
+	// may reference dangling data and function incorrectly in cases such as FindBlueprintforNodeChecked
+	for (FGraphPanelSelectionSet::TConstIterator NodeIt( SelectedNodes ); NodeIt; ++NodeIt)
+	{
+		if (UEdGraphNode* Node = Cast<UEdGraphNode>(*NodeIt))
+		{
+			if (Node->CanUserDeleteNode())
+			{
+				auto CloseAllDocumentsTab = [this](const UEdGraphNode* InNode)
+				{
+					TArray<UObject*> NodesToClose;
+					GetObjectsWithOuter(InNode, NodesToClose);
+					for (UObject* Node : NodesToClose)
+					{
+						UEdGraph* NodeGraph = Cast<UEdGraph>(Node);
+						if (NodeGraph)
+						{
+							CloseDocumentTab(NodeGraph);
+						}
+					}
+				};
+				
+				if (Node->GetSubGraphs().Num() > 0)
+				{
+					CloseAllDocumentsTab(Node);
+				}
+			}
+		}
+	}
+
+	// Now remove the selected nodes
+	for (FGraphPanelSelectionSet::TConstIterator NodeIt( SelectedNodes ); NodeIt; ++NodeIt)
+	{
+		if (UEdGraphNode* Node = Cast<UEdGraphNode>(*NodeIt))
+		{
+			if (Node->CanUserDeleteNode())
+			{
+				if (Node->GetSubGraphs().Num() > 0)
+				{
+					DocumentManager->CleanInvalidTabs();
+				}
+
+				FBlueprintEditorUtils::RemoveNode(nullptr, Node);
+			}
+		}
+	}
+}
+
+bool FGraphEditor::CanDeleteSelectedNodes()
+{
+	const FGraphPanelSelectionSet SelectedNodes = GetSelectedNodes();
+
+	bool bCanUserDeleteNode = false;
+
+	if(IsEditable(GetFocusedGraph()) && SelectedNodes.Num() > 0)
+	{
+		for(UObject* NodeObject : SelectedNodes)
+		{
+			// If any nodes allow deleting, then do not disable the delete option
+			UEdGraphNode* Node = Cast<UEdGraphNode>(NodeObject);
+			if(Node->CanUserDeleteNode())
+			{
+				bCanUserDeleteNode = true;
+				break;
+			}
+		}
+	}
+
+	return bCanUserDeleteNode;
+}
+
+}
+
+#undef LOCTEXT_NAMESPACE
\ No newline at end of file
diff --git a/Engine/Plugins/Experimental/DataInterfaceGraph/Source/DataInterfaceGraphEditor/Private/DataInterfaceGraphEditor.h b/Engine/Plugins/Experimental/DataInterfaceGraph/Source/DataInterfaceGraphEditor/Private/DataInterfaceGraphEditor.h
new file mode 100644
--- /dev/null
+++ b/Engine/Plugins/Experimental/DataInterfaceGraph/Source/DataInterfaceGraphEditor/Private/DataInterfaceGraphEditor.h
@@ -0,0 +1,123 @@
+﻿// Copyright Epic Games, Inc. All Rights Reserved.
+
+#pragma once
+
+#include "CoreMinimal.h"
+#include "WorkflowOrientedApp/WorkflowCentricApplication.h"
+#include "GraphEditor.h"
+
+enum class ERigVMGraphNotifType : uint8;
+class UDataInterfaceGraph;
+class UDataInterfaceGraph_EditorData;
+class FDocumentTracker;
+class FDocumentTabFactory;
+class FTabInfo;
+class FTabManager;
+class IToolkitHost;
+class URigVMGraph;
+class URigVMController;
+class UDetailsViewWrapperObject;
+
+namespace UE::DataInterfaceGraphEditor
+{
+
+namespace Modes
+{
+	extern const FName GraphEditor;
+}
+
+namespace Tabs
+{
+	extern const FName Details;
+	extern const FName Document;
+}
+
+class FGraphEditor : public FWorkflowCentricApplication
+{
+public:
+	/** Edits the specified asset */
+	void InitEditor(const EToolkitMode::Type InMode, const TSharedPtr<IToolkitHost>& InInitToolkitHost, UDataInterfaceGraph* InDataInterfaceGraph);
+
+private:
+	friend class FGraphEditorMode;
+	
+	// IToolkit interface
+	virtual void RegisterTabSpawners(const TSharedRef<FTabManager>& InTabManager) override;
+	virtual void UnregisterTabSpawners(const TSharedRef<FTabManager>& InTabManager) override;
+	virtual FName GetToolkitFName() const override;
+	virtual FText GetBaseToolkitName() const override;
+	virtual FString GetWorldCentricTabPrefix() const override;
+	virtual FLinearColor GetWorldCentricTabColorScale() const override;
+	virtual void InitToolMenuContext(FToolMenuContext& InMenuContext) override;
+
+	void BindCommands();
+	
+	void ExtendMenu();
+
+	void ExtendToolbar();
+
+	void OnGraphEditorFocused(TSharedRef<SGraphEditor> InGraphEditor);
+
+	void OnGraphEditorBackgrounded(TSharedRef<SGraphEditor> InGraphEditor);
+	
+	UEdGraph* GetFocusedGraph() const;
+
+	URigVMGraph* GetFocusedVMGraph() const;
+
+	URigVMController* GetFocusedVMController() const;
+	
+	// Create new graph editor widget for the supplied document container
+	TSharedRef<SGraphEditor> CreateGraphEditorWidget(TSharedRef<FTabInfo> InTabInfo, UEdGraph* InGraph);
+
+	// Create the action menu for the specified graph
+	FActionMenuContent OnCreateGraphActionMenu(UEdGraph* InGraph, const FVector2D& InNodePosition, const TArray<UEdGraphPin*>& InDraggedPins, bool bAutoExpand, SGraphEditor::FActionMenuClosed InOnMenuClosed);
+
+	// Handle node titles being changed
+	void OnNodeTitleCommitted(const FText& NewText, ETextCommit::Type CommitInfo, UEdGraphNode* NodeBeingChanged);
+
+	void CloseDocumentTab(const UObject* DocumentID);
+
+	bool InEditingMode() const;
+
+	bool IsEditable(UEdGraph* InGraph) const;
+	
+	FGraphPanelSelectionSet GetSelectedNodes() const;
+	
+	void DeleteSelectedNodes();
+	
+	bool CanDeleteSelectedNodes();
+
+	TSharedPtr<SGraphEditor> GetGraphEditor(UEdGraph* InEdGraph) const;
+	
+	void HandleModifiedEvent(ERigVMGraphNotifType InNotifType, URigVMGraph* InGraph, UObject* InSubject);
+
+	void HandleDetailsViewCreated(TSharedRef<IDetailsView> InDetailsView)
+	{
+		DetailsView = InDetailsView;
+	}
+	
+	// The asset we are editing
+	UDataInterfaceGraph* DataInterfaceGraph = nullptr;
+
+	// Cached editor data ptr
+	UDataInterfaceGraph_EditorData* DataInterfaceGraph_EditorData = nullptr;
+	
+	// Currently focused graph editor
+	TWeakPtr<SGraphEditor> FocusedGraphEdPtr;
+
+	// Document tracker
+	TSharedPtr<FDocumentTracker> DocumentManager;
+
+	// Factory for document tabs
+	TWeakPtr<FDocumentTabFactory> GraphEditorTabFactoryPtr;
+
+	// Command list for the graph editor
+	TSharedPtr<FUICommandList> GraphEditorCommands;
+
+	// Our details panel
+	TSharedPtr<IDetailsView> DetailsView;
+
+	TArray<TStrongObjectPtr<UDetailsViewWrapperObject>> WrapperObjects;
+};
+
+}
\ No newline at end of file
diff --git a/Engine/Plugins/Experimental/DataInterfaceGraph/Source/DataInterfaceGraphEditor/Private/DataInterfaceGraphEditorMode.cpp b/Engine/Plugins/Experimental/DataInterfaceGraph/Source/DataInterfaceGraphEditor/Private/DataInterfaceGraphEditorMode.cpp
new file mode 100644
--- /dev/null
+++ b/Engine/Plugins/Experimental/DataInterfaceGraph/Source/DataInterfaceGraphEditor/Private/DataInterfaceGraphEditorMode.cpp
@@ -0,0 +1,273 @@
+﻿// Copyright Epic Games, Inc. All Rights Reserved.
+
+#include "DataInterfaceGraphEditorMode.h"
+#include "DataInterfaceGraphEditor.h"
+#include "PropertyEditorModule.h"
+#include "Widgets/Docking/SDockTab.h"
+#include "EdGraph/EdGraph.h"
+#include "EdGraph/EdGraphSchema.h"
+#include "Widgets/Layout/SSpacer.h"
+#include "Modules/ModuleManager.h"
+
+#define LOCTEXT_NAMESPACE "GraphEditorMode"
+
+namespace UE::DataInterfaceGraphEditor
+{
+
+struct FGraphTabHistory : public FGenericTabHistory
+{
+public:
+	FGraphTabHistory(TSharedPtr<FDocumentTabFactory> InFactory, TSharedPtr<FTabPayload> InPayload)
+		: FGenericTabHistory(InFactory, InPayload)
+		, SavedLocation(FVector2D::ZeroVector)
+		, SavedZoomAmount(INDEX_NONE)
+	{
+	}
+
+private:
+	// FGenericTabHistory interface
+	virtual void EvokeHistory(TSharedPtr<::FTabInfo> InTabInfo, bool bPrevTabMatches) override
+	{
+		FWorkflowTabSpawnInfo SpawnInfo;
+		SpawnInfo.Payload = Payload;
+		SpawnInfo.TabInfo = InTabInfo;
+
+		if(bPrevTabMatches)
+		{
+			TSharedPtr<SDockTab> DockTab = InTabInfo->GetTab().Pin();
+			GraphEditor = StaticCastSharedRef<SGraphEditor>(DockTab->GetContent());
+		}
+		else
+		{
+			TSharedRef<SGraphEditor> GraphEditorRef = StaticCastSharedRef<SGraphEditor>(FactoryPtr.Pin()->CreateTabBody(SpawnInfo));
+			GraphEditor = GraphEditorRef;
+			FactoryPtr.Pin()->UpdateTab(InTabInfo->GetTab().Pin(), SpawnInfo, GraphEditorRef);
+		}
+	}
+	
+	virtual void SaveHistory() override
+	{
+		if (IsHistoryValid())
+		{
+			check(GraphEditor.IsValid());
+			GraphEditor.Pin()->GetViewLocation(SavedLocation, SavedZoomAmount);
+			GraphEditor.Pin()->GetViewBookmark(SavedBookmarkId);
+		}
+	}
+	
+	virtual void RestoreHistory() override
+	{
+		if (IsHistoryValid())
+		{
+			check(GraphEditor.IsValid());
+			GraphEditor.Pin()->SetViewLocation(SavedLocation, SavedZoomAmount, SavedBookmarkId);
+		}
+	}
+	
+	/** The graph editor represented by this history node. While this node is inactive, the graph editor is invalid */
+	TWeakPtr<SGraphEditor> GraphEditor;
+	
+	/** Saved location the graph editor was at when this history node was last visited */
+	FVector2D SavedLocation;
+	
+	/** Saved zoom the graph editor was at when this history node was last visited */
+	float SavedZoomAmount;
+	
+	/** Saved bookmark ID the graph editor was at when this history node was last visited */
+	FGuid SavedBookmarkId;
+};
+
+FGraphEditorSummoner::FGraphEditorSummoner(TSharedPtr<FWorkflowCentricApplication> InHostingApp)
+	: FDocumentTabFactoryForObjects<UEdGraph>(Tabs::Document, InHostingApp)
+{
+}
+
+void FGraphEditorSummoner::OnTabActivated(TSharedPtr<SDockTab> Tab) const
+{
+	TSharedRef<SGraphEditor> GraphEditor = StaticCastSharedRef<SGraphEditor>(Tab->GetContent());
+	OnGraphEditorFocusedDelegate.ExecuteIfBound(GraphEditor);
+}
+
+void FGraphEditorSummoner::OnTabBackgrounded(TSharedPtr<SDockTab> Tab) const
+{
+	TSharedRef<SGraphEditor> GraphEditor = StaticCastSharedRef<SGraphEditor>(Tab->GetContent());
+	OnGraphEditorBackgroundedDelegate.ExecuteIfBound(GraphEditor);
+}
+
+void FGraphEditorSummoner::OnTabRefreshed(TSharedPtr<SDockTab> Tab) const
+{
+	TSharedRef<SGraphEditor> GraphEditor = StaticCastSharedRef<SGraphEditor>(Tab->GetContent());
+	GraphEditor->NotifyGraphChanged();
+}
+
+void FGraphEditorSummoner::SaveState(TSharedPtr<SDockTab> Tab, TSharedPtr<FTabPayload> Payload) const
+{
+	TSharedRef<SGraphEditor> GraphEditor = StaticCastSharedRef<SGraphEditor>(Tab->GetContent());
+
+	FVector2D ViewLocation;
+	float ZoomAmount;
+	GraphEditor->GetViewLocation(ViewLocation, ZoomAmount);
+
+	UEdGraph* Graph = Payload->IsValid() ? FTabPayload_UObject::CastChecked<UEdGraph>(Payload) : nullptr;
+
+	OnSaveGraphStateDelegate.ExecuteIfBound(Graph, ViewLocation, ZoomAmount);
+}
+
+TAttribute<FText> FGraphEditorSummoner::ConstructTabNameForObject(UEdGraph* DocumentID) const
+{
+	if (DocumentID)
+	{
+		if (const UEdGraphSchema* Schema = DocumentID->GetSchema())
+		{
+			FGraphDisplayInfo Info;
+			Schema->GetGraphDisplayInformation(*DocumentID, /*out*/ Info);
+			return Info.DisplayName;
+		}
+		else
+		{
+			// if we don't have a schema, we're dealing with a malformed (or incomplete graph)...
+			// possibly in the midst of some transaction - here we return the object's outer path 
+			// so we can at least get some context as to which graph we're referring
+			return FText::FromString(DocumentID->GetPathName());
+		}
+	}
+
+	return LOCTEXT("UnknownGraphName", "Unknown");
+}
+
+TSharedRef<SWidget> FGraphEditorSummoner::CreateTabBodyForObject(const FWorkflowTabSpawnInfo& Info, UEdGraph* DocumentID) const
+{
+	check(Info.TabInfo.IsValid());
+	
+	if(OnCreateGraphEditorWidgetDelegate.IsBound())
+	{
+		return OnCreateGraphEditorWidgetDelegate.Execute(Info.TabInfo.ToSharedRef(), DocumentID);
+	}
+	else
+	{
+		return SNew(SSpacer);
+	}
+}
+
+const FSlateBrush* FGraphEditorSummoner::GetTabIconForObject(const FWorkflowTabSpawnInfo& Info, UEdGraph* DocumentID) const
+{
+	if(OnGetTabIconDelegate.IsBound())
+	{
+		return OnGetTabIconDelegate.Execute(DocumentID);
+	}
+	else
+	{
+		return FAppStyle::Get().GetBrush(TEXT("GraphEditor.EventGraph_16x"));
+	}
+}
+
+TSharedRef<FGenericTabHistory> FGraphEditorSummoner::CreateTabHistoryNode(TSharedPtr<FTabPayload> Payload)
+{
+	return MakeShared<FGraphTabHistory>(SharedThis(this), Payload);
+}
+
+class FDetailsTabSummoner : public FWorkflowTabFactory
+{
+public:
+	FDetailsTabSummoner(TSharedPtr<FAssetEditorToolkit> InHostingApp, FOnDetailsViewCreated InOnDetailsViewCreated)
+		: FWorkflowTabFactory(Tabs::Details, InHostingApp)
+	{
+		TabLabel = LOCTEXT("DetailsTabLabel", "Details");
+		TabIcon = FSlateIcon("EditorStyle", "LevelEditor.Tabs.Details");
+		ViewMenuDescription = LOCTEXT("DetailsTabMenuDescription", "Details");
+		ViewMenuTooltip = LOCTEXT("DetailsTabToolTip", "Shows the details tab for selected objects.");
+		bIsSingleton = true;
+
+		FDetailsViewArgs DetailsViewArgs;
+		DetailsViewArgs.NameAreaSettings = FDetailsViewArgs::HideNameArea;
+		DetailsViewArgs.bHideSelectionTip = true;
+
+		FPropertyEditorModule& PropertyEditorModule = FModuleManager::LoadModuleChecked<FPropertyEditorModule>("PropertyEditor");
+		DetailsView = PropertyEditorModule.CreateDetailView(DetailsViewArgs);
+		
+		InOnDetailsViewCreated.ExecuteIfBound(DetailsView.ToSharedRef());
+	}
+
+	TSharedPtr<IDetailsView> GetDetailsView() const
+	{
+		return DetailsView;
+	}
+
+private:
+	virtual TSharedRef<SWidget> CreateTabBody(const FWorkflowTabSpawnInfo& Info) const override
+	{
+		return DetailsView.ToSharedRef();
+	}
+	
+	virtual FText GetTabToolTipText(const FWorkflowTabSpawnInfo& Info) const override
+	{
+		return ViewMenuTooltip;
+	}
+	
+	TSharedPtr<IDetailsView> DetailsView;
+};
+
+FGraphEditorMode::FGraphEditorMode(TSharedRef<FWorkflowCentricApplication> InHostingApp)
+	: FApplicationMode(Modes::GraphEditor)
+	, HostingAppPtr(InHostingApp)
+{
+	HostingAppPtr = InHostingApp;
+
+	TSharedRef<FGraphEditor> GraphEditor = StaticCastSharedRef<FGraphEditor>(InHostingApp);
+	
+	TabFactories.RegisterFactory(MakeShared<FDetailsTabSummoner>(GraphEditor, FOnDetailsViewCreated::CreateSP(&GraphEditor.Get(), &FGraphEditor::HandleDetailsViewCreated)));
+
+	TabLayout = FTabManager::NewLayout("Standalone_DataInterfaceGraphEditor_Layout_v1.0")
+		->AddArea
+		(
+			FTabManager::NewPrimaryArea()
+			->SetOrientation(Orient_Vertical)
+			->Split
+			(	
+				FTabManager::NewSplitter()
+				->SetSizeCoefficient(1.0f)
+				->SetOrientation(Orient_Horizontal)
+				->Split
+				(
+					FTabManager::NewStack()
+					->SetSizeCoefficient(0.8f)
+					->SetHideTabWell(false)
+					->AddTab(Tabs::Document, ETabState::ClosedTab)
+				)
+				->Split
+				(
+					FTabManager::NewStack()
+					->SetSizeCoefficient(0.2f)
+					->SetHideTabWell(false)
+					->AddTab(Tabs::Details, ETabState::OpenedTab)
+				)
+			)
+		);
+
+	GraphEditor->RegisterModeToolbarIfUnregistered(GetModeName());
+}
+
+void FGraphEditorMode::RegisterTabFactories(TSharedPtr<FTabManager> InTabManager)
+{
+	TSharedPtr<FWorkflowCentricApplication> HostingApp = HostingAppPtr.Pin();
+	HostingApp->PushTabFactories(TabFactories);
+
+	FApplicationMode::RegisterTabFactories(InTabManager);
+}
+
+void FGraphEditorMode::AddTabFactory(FCreateWorkflowTabFactory FactoryCreator)
+{
+	if (FactoryCreator.IsBound())
+	{
+		TabFactories.RegisterFactory(FactoryCreator.Execute(HostingAppPtr.Pin()));
+	}
+}
+
+void FGraphEditorMode::RemoveTabFactory(FName TabFactoryID)
+{
+	TabFactories.UnregisterFactory(TabFactoryID);
+}
+
+}
+
+#undef LOCTEXT_NAMESPACE
\ No newline at end of file
diff --git a/Engine/Plugins/Experimental/DataInterfaceGraph/Source/DataInterfaceGraphEditor/Private/DataInterfaceGraphEditorMode.h b/Engine/Plugins/Experimental/DataInterfaceGraph/Source/DataInterfaceGraphEditor/Private/DataInterfaceGraphEditorMode.h
new file mode 100644
--- /dev/null
+++ b/Engine/Plugins/Experimental/DataInterfaceGraph/Source/DataInterfaceGraphEditor/Private/DataInterfaceGraphEditorMode.h
@@ -0,0 +1,95 @@
+﻿// Copyright Epic Games, Inc. All Rights Reserved.
+
+#pragma once
+
+#include "CoreMinimal.h"
+#include "WorkflowOrientedApp/ApplicationMode.h"
+#include "WorkflowOrientedApp/WorkflowTabManager.h"
+#include "WorkflowOrientedApp/WorkflowUObjectDocuments.h"
+
+class FWorkflowCentricApplication;
+class SGraphEditor;
+
+namespace UE::DataInterfaceGraphEditor
+{
+
+DECLARE_DELEGATE_OneParam(FOnDetailsViewCreated, TSharedRef<IDetailsView>);
+
+struct FGraphEditorSummoner : public FDocumentTabFactoryForObjects<UEdGraph>
+{
+public:
+	// Delegate called when a graph editor widget is created
+	DECLARE_DELEGATE_RetVal_TwoParams(TSharedRef<SGraphEditor>, FOnCreateGraphEditorWidget, TSharedRef<FTabInfo>, UEdGraph*);
+
+	// Delegate called when a graph editor is focused
+	DECLARE_DELEGATE_OneParam(FOnGraphEditorFocused, TSharedRef<SGraphEditor>);
+
+	// Delegate called when a graph editor is backgrounded
+	DECLARE_DELEGATE_OneParam(FOnGraphEditorBackgrounded, TSharedRef<SGraphEditor>);
+
+	// Delegate called to save the state of a graph
+	DECLARE_DELEGATE_ThreeParams(FOnSaveGraphState, UEdGraph*, FVector2D, float);
+
+	// Delegate called to save the state of a graph
+	DECLARE_DELEGATE_RetVal_OneParam(const FSlateBrush*, FOnGetTabIcon, UEdGraph*);
+	
+	FGraphEditorSummoner(TSharedPtr<FWorkflowCentricApplication> InHostingApp);
+
+	FOnCreateGraphEditorWidget& OnCreateGraphEditorWidget() { return OnCreateGraphEditorWidgetDelegate; }
+	FOnGraphEditorFocused& OnGraphEditorFocused() { return OnGraphEditorFocusedDelegate; }
+	FOnGraphEditorBackgrounded& OnGraphEditorBackgrounded() { return OnGraphEditorBackgroundedDelegate; }
+	FOnSaveGraphState& OnSaveGraphState() { return OnSaveGraphStateDelegate; }
+	FOnGetTabIcon& OnGetTabIcon() { return OnGetTabIconDelegate; }
+
+private:
+	// FWorkflowTabFactory interface
+	virtual void OnTabActivated(TSharedPtr<SDockTab> Tab) const override;
+	virtual void OnTabBackgrounded(TSharedPtr<SDockTab> Tab) const override;
+	virtual void OnTabRefreshed(TSharedPtr<SDockTab> Tab) const override;
+	virtual void SaveState(TSharedPtr<SDockTab> Tab, TSharedPtr<FTabPayload> Payload) const override;
+	
+	// FDocumentTabFactoryForObjects interface
+	virtual TAttribute<FText> ConstructTabNameForObject(UEdGraph* DocumentID) const override;
+	virtual TSharedRef<SWidget> CreateTabBodyForObject(const FWorkflowTabSpawnInfo& Info, UEdGraph* DocumentID) const override;
+	virtual const FSlateBrush* GetTabIconForObject(const FWorkflowTabSpawnInfo& Info, UEdGraph* DocumentID) const override;
+	virtual TSharedRef<FGenericTabHistory> CreateTabHistoryNode(TSharedPtr<FTabPayload> Payload) override;
+
+	// The hosting app
+	TWeakPtr<FWorkflowCentricApplication> HostingAppPtr;
+
+	// Delegate called when a graph editor widget is created
+	FOnCreateGraphEditorWidget OnCreateGraphEditorWidgetDelegate;
+
+	// Delegate called when a graph editor is focused
+	FOnGraphEditorFocused OnGraphEditorFocusedDelegate;
+
+	// Delegate called when a graph editor is backgrounded
+	FOnGraphEditorBackgrounded OnGraphEditorBackgroundedDelegate;
+
+	// Delegate called to save the state of a graph
+	FOnSaveGraphState OnSaveGraphStateDelegate;
+
+	// Delegate called to get a tab icon for a graph
+	FOnGetTabIcon OnGetTabIconDelegate;
+};
+
+class FGraphEditorMode : public FApplicationMode
+{
+public:
+	FGraphEditorMode(TSharedRef<FWorkflowCentricApplication> InHostingApp);
+
+private:
+	// FApplicationMode interface
+	virtual void RegisterTabFactories(TSharedPtr<FTabManager> InTabManager) override;
+	virtual void AddTabFactory(FCreateWorkflowTabFactory FactoryCreator) override;
+	virtual void RemoveTabFactory(FName TabFactoryID) override;
+	
+	// The hosting app
+	TWeakPtr<FWorkflowCentricApplication> HostingAppPtr;
+
+	// The tab factories we support
+	FWorkflowAllowedTabSet TabFactories;
+};
+
+}
+
diff --git a/Engine/Plugins/Experimental/DataInterfaceGraph/Source/DataInterfaceGraphEditor/Private/DataInterfaceGraphEditorModule.cpp b/Engine/Plugins/Experimental/DataInterfaceGraph/Source/DataInterfaceGraphEditor/Private/DataInterfaceGraphEditorModule.cpp
new file mode 100644
--- /dev/null
+++ b/Engine/Plugins/Experimental/DataInterfaceGraph/Source/DataInterfaceGraphEditor/Private/DataInterfaceGraphEditorModule.cpp
@@ -0,0 +1,42 @@
+// Copyright Epic Games, Inc. All Rights Reserved.
+
+#include "DataInterfaceGraphEditorModule.h"
+#include "IAssetTools.h"
+#include "AssetTypeActions_DataInterfaceGraph.h"
+#include "DataInterfacePropertyTypeCustomization.h"
+
+namespace UE::DataInterfaceGraphEditor
+{
+
+void FModule::StartupModule()
+{
+	IAssetTools& AssetTools = FModuleManager::LoadModuleChecked<FAssetToolsModule>("AssetTools").Get();
+	AssetTypeActions_DataInterfaceGraph = MakeShared<UE::DataInterfaceGraphEditor::FAssetTypeActions_DataInterfaceGraph>();
+	AssetTools.RegisterAssetTypeActions(AssetTypeActions_DataInterfaceGraph.ToSharedRef());
+
+	DataInterfacePropertyTypeIdentifier = MakeShared<UE::DataInterfaceGraphEditor::FPropertyTypeIdentifier>();
+	FPropertyEditorModule& PropertyModule = FModuleManager::LoadModuleChecked<FPropertyEditorModule>("PropertyEditor");
+	PropertyModule.RegisterCustomPropertyTypeLayout(
+		FInterfaceProperty::StaticClass()->GetFName(),
+		FOnGetPropertyTypeCustomizationInstance::CreateLambda([] { return MakeShared<UE::DataInterfaceGraphEditor::FDataInterfacePropertyTypeCustomization>(); }),
+		DataInterfacePropertyTypeIdentifier);
+}
+
+void FModule::ShutdownModule()
+{
+	if(FModuleManager::Get().IsModuleLoaded("AssetTools"))
+	{
+		IAssetTools& AssetTools = FModuleManager::GetModuleChecked<FAssetToolsModule>("AssetTools").Get();
+		AssetTools.UnregisterAssetTypeActions(AssetTypeActions_DataInterfaceGraph.ToSharedRef());
+	}
+	
+	if(FModuleManager::Get().IsModuleLoaded("PropertyEditor"))
+	{
+		FPropertyEditorModule& PropertyModule = FModuleManager::LoadModuleChecked<FPropertyEditorModule>("PropertyEditor");
+		PropertyModule.UnregisterCustomPropertyTypeLayout(FInterfaceProperty::StaticClass()->GetFName(), DataInterfacePropertyTypeIdentifier);
+	}
+}
+
+}
+
+IMPLEMENT_MODULE(UE::DataInterfaceGraphEditor::FModule, DataInterfaceGraphEditor);
\ No newline at end of file
diff --git a/Engine/Plugins/Experimental/DataInterfaceGraph/Source/DataInterfaceGraphEditor/Private/DataInterfaceGraphEditorModule.h b/Engine/Plugins/Experimental/DataInterfaceGraph/Source/DataInterfaceGraphEditor/Private/DataInterfaceGraphEditorModule.h
new file mode 100644
--- /dev/null
+++ b/Engine/Plugins/Experimental/DataInterfaceGraph/Source/DataInterfaceGraphEditor/Private/DataInterfaceGraphEditorModule.h
@@ -0,0 +1,25 @@
+﻿// Copyright Epic Games, Inc. All Rights Reserved.
+
+#pragma once
+
+#include "CoreMinimal.h"
+#include "Modules/ModuleManager.h"
+
+namespace UE::DataInterfaceGraphEditor
+{
+
+class FAssetTypeActions_DataInterfaceGraph;
+class FDataInterfacePropertyTypeCustomization;
+class FPropertyTypeIdentifier;
+
+class FModule : public IModuleInterface
+{
+private:
+	virtual void StartupModule() override;
+	virtual void ShutdownModule() override;
+
+	TSharedPtr<UE::DataInterfaceGraphEditor::FAssetTypeActions_DataInterfaceGraph> AssetTypeActions_DataInterfaceGraph;
+	TSharedPtr<UE::DataInterfaceGraphEditor::FPropertyTypeIdentifier> DataInterfacePropertyTypeIdentifier;
+};
+
+}
\ No newline at end of file
diff --git a/Engine/Plugins/Experimental/DataInterfaceGraph/Source/DataInterfaceGraphEditor/Private/DataInterfaceGraphEditorSchemaActions.cpp b/Engine/Plugins/Experimental/DataInterfaceGraph/Source/DataInterfaceGraphEditor/Private/DataInterfaceGraphEditorSchemaActions.cpp
new file mode 100644
--- /dev/null
+++ b/Engine/Plugins/Experimental/DataInterfaceGraph/Source/DataInterfaceGraphEditor/Private/DataInterfaceGraphEditorSchemaActions.cpp
@@ -0,0 +1,61 @@
+﻿// Copyright Epic Games, Inc. All Rights Reserved.
+
+#include "DataInterfaceGraphEditorSchemaActions.h"
+#include "DataInterfaceEditorUtils.h"
+#include "DataInterfaceGraph_EdGraph.h"
+#include "DataInterfaceGraph_EdGraphNode.h"
+#include "SGraphActionMenu.h"
+#include "DataInterfaceUncookedOnlyUtils.h"
+#include "DataInterfaceGraph_EditorData.h"
+#include "RigVMModel/Nodes/RigVMUnitNode.h"
+#include "Settings/ControlRigSettings.h"
+
+UEdGraphNode* FDataInterfaceGraphSchemaAction_RigUnit::PerformAction(UEdGraph* ParentGraph, TArray<UEdGraphPin*>& FromPins, const FVector2D Location, bool bSelectNewNode)
+{
+	UDataInterfaceGraph_EditorData* EditorData = ParentGraph->GetTypedOuter<UDataInterfaceGraph_EditorData>();
+	UDataInterfaceGraph_EdGraphNode* NewNode = nullptr;
+	UDataInterfaceGraph_EdGraph* EdGraph = Cast<UDataInterfaceGraph_EdGraph>(ParentGraph);
+
+	if (EditorData != nullptr && EdGraph != nullptr)
+	{
+		FName Name = UE::DataInterfaceGraphEditor::FUtils::ValidateName(EditorData, StructTemplate->GetFName().ToString());
+		URigVMController* Controller = EditorData->GetRigVMController(ParentGraph);
+
+		Controller->OpenUndoBracket(FString::Printf(TEXT("Add '%s' Node"), *Name.ToString()));
+
+		FRigVMUnitNodeCreatedContext& UnitNodeCreatedContext = Controller->GetUnitNodeCreatedContext();
+		FRigVMUnitNodeCreatedContext::FScope ReasonScope(UnitNodeCreatedContext, ERigVMNodeCreatedReason::NodeSpawner);
+
+		if (URigVMUnitNode* ModelNode = Controller->AddUnitNode(StructTemplate, FRigUnit::GetMethodName(), Location, Name.ToString(), true, false))
+		{
+			NewNode = Cast<UDataInterfaceGraph_EdGraphNode>(EdGraph->FindNodeForModelNodeName(ModelNode->GetFName()));
+			check(NewNode);
+
+			if (NewNode)
+			{
+				Controller->ClearNodeSelection(true);
+				Controller->SelectNode(ModelNode, true, true);
+			}
+
+#if WITH_EDITORONLY_DATA
+
+			// @TODO: switch over to custom settings here
+			const FControlRigSettingsPerPinBool* ExpansionMapPtr = UControlRigEditorSettings::Get()->RigUnitPinExpansion.Find(ModelNode->GetScriptStruct()->GetName());
+			if (ExpansionMapPtr)
+			{
+				const FControlRigSettingsPerPinBool& ExpansionMap = *ExpansionMapPtr;
+
+				for (const TPair<FString, bool>& Pair : ExpansionMap.Values)
+				{
+					FString PinPath = FString::Printf(TEXT("%s.%s"), *ModelNode->GetName(), *Pair.Key);
+					Controller->SetPinExpansion(PinPath, Pair.Value, true);
+				}
+			}
+		}
+#endif
+		
+		Controller->CloseUndoBracket();
+	}
+
+	return NewNode;
+}
diff --git a/Engine/Plugins/Experimental/DataInterfaceGraph/Source/DataInterfaceGraphEditor/Private/DataInterfaceGraphEditorSchemaActions.h b/Engine/Plugins/Experimental/DataInterfaceGraph/Source/DataInterfaceGraphEditor/Private/DataInterfaceGraphEditorSchemaActions.h
new file mode 100644
--- /dev/null
+++ b/Engine/Plugins/Experimental/DataInterfaceGraph/Source/DataInterfaceGraphEditor/Private/DataInterfaceGraphEditorSchemaActions.h
@@ -0,0 +1,55 @@
+﻿// Copyright Epic Games, Inc. All Rights Reserved.
+
+#pragma once
+
+#include "CoreMinimal.h"
+#include "EdGraph/EdGraphSchema.h"
+#include "Styling/AppStyle.h"
+#include "DataInterfaceGraphEditorSchemaActions.generated.h"
+
+struct FSlateBrush;
+
+USTRUCT()
+struct FDataInterfaceGraphSchemaAction : public FEdGraphSchemaAction
+{
+	GENERATED_BODY()
+
+	FDataInterfaceGraphSchemaAction() = default;
+	
+	FDataInterfaceGraphSchemaAction(FText InNodeCategory, FText InMenuDesc, FText InToolTip, FText InKeywords = FText::GetEmpty())
+		: FEdGraphSchemaAction(MoveTemp(InNodeCategory), MoveTemp(InMenuDesc), MoveTemp(InToolTip), 0, MoveTemp(InKeywords))
+	{
+	}
+
+	virtual const FSlateBrush* GetIconBrush() const
+	{
+		return FAppStyle::Get().GetBrush("NoBrush");
+	}
+
+	virtual const FLinearColor& GetIconColor() const
+	{
+		static const FLinearColor DefaultColor;
+		return DefaultColor;
+	}
+};
+
+USTRUCT()
+struct FDataInterfaceGraphSchemaAction_RigUnit : public FDataInterfaceGraphSchemaAction
+{
+	GENERATED_BODY()
+
+	FDataInterfaceGraphSchemaAction_RigUnit() = default;
+	
+	FDataInterfaceGraphSchemaAction_RigUnit(UScriptStruct* InStructTemplate, FText InNodeCategory, FText InMenuDesc, FText InToolTip, FText InKeywords = FText::GetEmpty())
+		: FDataInterfaceGraphSchemaAction(InNodeCategory, InMenuDesc, InToolTip, InKeywords)
+		, StructTemplate(InStructTemplate)
+	{}
+
+	// FEdGraphSchemaAction Interface
+	virtual UEdGraphNode* PerformAction(UEdGraph* ParentGraph, TArray<UEdGraphPin*>& FromPins, const FVector2D Location, bool bSelectNewNode = true) override;
+	virtual UEdGraphNode* PerformAction(UEdGraph* ParentGraph, UEdGraphPin* FromPin, const FVector2D Location, bool bSelectNewNode = true) { return nullptr; }
+
+private:
+	// The script struct for our rig unit
+	UScriptStruct* StructTemplate = nullptr;
+};
diff --git a/Engine/Plugins/Experimental/DataInterfaceGraph/Source/DataInterfaceGraphEditor/Private/DataInterfaceGraphFactory.cpp b/Engine/Plugins/Experimental/DataInterfaceGraph/Source/DataInterfaceGraphEditor/Private/DataInterfaceGraphFactory.cpp
new file mode 100644
--- /dev/null
+++ b/Engine/Plugins/Experimental/DataInterfaceGraph/Source/DataInterfaceGraphEditor/Private/DataInterfaceGraphFactory.cpp
@@ -0,0 +1,94 @@
+﻿// Copyright Epic Games, Inc. All Rights Reserved.
+
+#include "DataInterfaceGraphFactory.h"
+#include "DataInterfaceGraph.h"
+#include "DataInterfaceGraph_EditorData.h"
+#include "DataInterfaceGraph_EdGraph.h"
+#include "DataInterfaceGraph_EdGraphSchema.h"
+#include "DataInterfaceUncookedOnlyUtils.h"
+#include "RigUnit_DataInterfaceBeginExecution.h"
+#include "RigUnit_DataInterfaceEndExecution.h"
+#include "Kismet2/BlueprintEditorUtils.h"
+#include "RigVMModel/RigVMController.h"
+#include "Units/RigUnit.h"
+
+UDataInterfaceGraphFactory::UDataInterfaceGraphFactory()
+{
+	bCreateNew = true;
+	bEditAfterNew = true;
+	SupportedClass = UDataInterfaceGraph::StaticClass();
+}
+
+bool UDataInterfaceGraphFactory::ConfigureProperties()
+{
+	return true;
+}
+
+UObject* UDataInterfaceGraphFactory::FactoryCreateNew(UClass* Class, UObject* InParent, FName Name, EObjectFlags Flags, UObject* Context, FFeedbackContext* Warn, FName CallingContext)
+{
+	UDataInterfaceGraph* NewGraph = NewObject<UDataInterfaceGraph>(InParent, Class, Name, Flags | RF_Public | RF_Standalone | RF_Transactional | RF_LoadCompleted);
+
+	// Create internal editor data & graphs
+	UDataInterfaceGraph_EditorData* EditorData = NewObject<UDataInterfaceGraph_EditorData>(NewGraph, TEXT("EditorData"));
+	NewGraph->EditorData = EditorData;
+	EditorData->Initialize(/*bRecompileVM*/false);
+
+	// Add initial execution unit
+	URigVMController* RootController = EditorData->GetRigVMController(EditorData->RootGraph);
+	URigVMUnitNode* MainEntryPointNode = RootController->AddUnitNode(FRigUnit_DataInterfaceBeginExecution::StaticStruct(), FRigUnit::GetMethodName(), FVector2D(-400.0f, 0.0f), FString(), false);
+	URigVMPin* BeginExecutePin = MainEntryPointNode->FindPin(GET_MEMBER_NAME_STRING_CHECKED(FRigUnit_DataInterfaceBeginExecution, ExecuteContext));
+	check(BeginExecutePin);
+	check(BeginExecutePin->GetDirection() == ERigVMPinDirection::Output);
+	
+	// Add function to function lib
+	URigVMController* FunctionLibraryController = EditorData->GetRigVMController(EditorData->RigVMFunctionLibrary);
+	EditorData->EntryPoint = FunctionLibraryController->AddFunctionToLibrary(UE::DataInterfaceGraph::EntryPointName, true, FVector2D::ZeroVector, false, false);
+	
+	// Add exposed pin for result to function
+	URigVMController* EntryPointController = EditorData->GetOrCreateRigVMController(EditorData->EntryPointGraph);
+	
+	// TODO: using float for now, but needs to be user-driven
+	{
+		const FString CPPType = TEXT("float");
+		const FName CPPTypeObjectPath = NAME_None;
+		const FString DefaultValue = TEXT("0.0");
+		
+		EntryPointController->AddExposedPin(UE::DataInterfaceGraph::ResultName, ERigVMPinDirection::Output, CPPType, CPPTypeObjectPath, DefaultValue, false, false);
+	}
+
+	// Add function call
+	URigVMFunctionReferenceNode* FunctionCallNode = RootController->AddFunctionReferenceNode(EditorData->EntryPoint, FVector2D::ZeroVector, UE::DataInterfaceGraph::EntryPointName.ToString(), false, false);
+	URigVMPin* FunctionExecutePin = FunctionCallNode->FindPin(FRigVMStruct::ExecuteContextName.ToString());
+	check(FunctionExecutePin);
+	check(FunctionExecutePin->GetDirection() == ERigVMPinDirection::IO);
+
+	// Link entry point to function call
+	RootController->AddLink(BeginExecutePin, FunctionExecutePin, false); 
+	
+	URigVMPin* FunctionResultPin = FunctionCallNode->FindPin(UE::DataInterfaceGraph::ResultName.ToString());
+	check(FunctionResultPin);
+	check(FunctionResultPin->GetDirection() == ERigVMPinDirection::Output);
+	
+	// Add end-execution unit of the correct type
+	// TODO: using a float for now, but need to use a registry to determine correct type
+	URigVMUnitNode* MainExitPointNode = RootController->AddUnitNode(FRigUnit_DataInterfaceEndExecution_Float::StaticStruct(), FRigUnit::GetMethodName(), FVector2D(400.0f, 0.0f), FString(), false);
+	URigVMPin* EndExecutePin = MainExitPointNode->FindPin(GET_MEMBER_NAME_STRING_CHECKED(FRigUnit_DataInterfaceEndExecution, ExecuteContext));
+	check(EndExecutePin);
+	check(EndExecutePin->GetDirection() == ERigVMPinDirection::Input);
+
+	URigVMPin* EndExecuteResultPin = MainExitPointNode->FindPin(GET_MEMBER_NAME_STRING_CHECKED(FRigUnit_DataInterfaceEndExecution_Float, Result));
+	check(EndExecuteResultPin);
+	check(EndExecuteResultPin->GetDirection() == ERigVMPinDirection::Input);
+	
+	// Link function call to exit point 
+	RootController->AddLink(FunctionExecutePin, EndExecutePin, false); 
+
+	// Link result pins 
+	RootController->AddLink(FunctionResultPin, EndExecuteResultPin, false); 
+	
+	// Compile the initial skeleton
+	UE::DataInterfaceGraphUncookedOnly::FUtils::Compile(NewGraph);
+	check(!EditorData->bErrorsDuringCompilation);
+	
+	return NewGraph;
+}
\ No newline at end of file
diff --git a/Engine/Plugins/Experimental/DataInterfaceGraph/Source/DataInterfaceGraphEditor/Private/DataInterfaceGraphFactory.h b/Engine/Plugins/Experimental/DataInterfaceGraph/Source/DataInterfaceGraphEditor/Private/DataInterfaceGraphFactory.h
new file mode 100644
--- /dev/null
+++ b/Engine/Plugins/Experimental/DataInterfaceGraph/Source/DataInterfaceGraphEditor/Private/DataInterfaceGraphFactory.h
@@ -0,0 +1,22 @@
+﻿// Copyright Epic Games, Inc. All Rights Reserved.
+
+#pragma once
+
+#include "CoreMinimal.h"
+#include "Factories/Factory.h"
+#include "DataInterfaceGraphFactory.generated.h"
+
+UCLASS(MinimalAPI)
+class UDataInterfaceGraphFactory : public UFactory
+{
+	GENERATED_BODY()
+
+	UDataInterfaceGraphFactory();
+	
+	virtual bool ConfigureProperties() override;
+	virtual UObject* FactoryCreateNew(UClass* Class, UObject* InParent, FName Name, EObjectFlags Flags, UObject* Context, FFeedbackContext* Warn, FName CallingContext) override;
+	virtual UObject* FactoryCreateNew(UClass* Class, UObject* InParent, FName Name, EObjectFlags Flags, UObject* Context, FFeedbackContext* Warn) override
+	{
+		return FactoryCreateNew(Class, InParent, Name, Flags, Context, Warn, NAME_None);
+	}
+};
\ No newline at end of file
diff --git a/Engine/Plugins/Experimental/DataInterfaceGraph/Source/DataInterfaceGraphEditor/Private/DataInterfacePropertyTypeCustomization.cpp b/Engine/Plugins/Experimental/DataInterfaceGraph/Source/DataInterfaceGraphEditor/Private/DataInterfacePropertyTypeCustomization.cpp
new file mode 100644
--- /dev/null
+++ b/Engine/Plugins/Experimental/DataInterfaceGraph/Source/DataInterfaceGraphEditor/Private/DataInterfacePropertyTypeCustomization.cpp
@@ -0,0 +1,52 @@
+﻿// Copyright Epic Games, Inc. All Rights Reserved.
+
+#include "DataInterfacePropertyTypeCustomization.h"
+#include "PropertyHandle.h"
+#include "IDataInterface.h"
+#include "DetailWidgetRow.h"
+#include "PropertyCustomizationHelpers.h"
+#include "Widgets/Text/STextBlock.h"
+
+#define LOCTEXT_NAMESPACE "DataInterfacePropertyTypeCustomization"
+
+namespace UE::DataInterfaceGraphEditor
+{
+
+bool FPropertyTypeIdentifier::IsPropertyTypeCustomized(const IPropertyHandle& PropertyHandle) const
+{
+	FInterfaceProperty* InterfaceProperty = CastField<FInterfaceProperty>(PropertyHandle.GetProperty());
+	return InterfaceProperty && InterfaceProperty->InterfaceClass->IsChildOf(UDataInterface::StaticClass());
+}
+
+void FDataInterfacePropertyTypeCustomization::CustomizeHeader(TSharedRef<IPropertyHandle> PropertyHandle, FDetailWidgetRow& HeaderRow, IPropertyTypeCustomizationUtils& CustomizationUtils)
+{
+	FInterfaceProperty* InterfaceProperty = CastField<FInterfaceProperty>(PropertyHandle->GetProperty());
+
+	HeaderRow
+	.NameContent()
+	[
+		PropertyHandle->CreatePropertyNameWidget()
+	]
+	.ValueContent()
+	[
+		SNew(SObjectPropertyEntryBox)
+		.PropertyHandle(PropertyHandle)
+		.DisplayUseSelected(false)
+		.OnShouldFilterAsset_Lambda([InterfaceProperty](const FAssetData& InAssetData)
+		{
+			if(UClass* Class = InAssetData.GetClass())
+			{
+				if(Class->ImplementsInterface(InterfaceProperty->InterfaceClass))
+				{
+					return false;
+				}
+			}
+
+			return true;
+		})
+	];
+}
+
+}
+
+#undef LOCTEXT_NAMESPACE
\ No newline at end of file
diff --git a/Engine/Plugins/Experimental/DataInterfaceGraph/Source/DataInterfaceGraphEditor/Private/DataInterfacePropertyTypeCustomization.h b/Engine/Plugins/Experimental/DataInterfaceGraph/Source/DataInterfaceGraphEditor/Private/DataInterfacePropertyTypeCustomization.h
new file mode 100644
--- /dev/null
+++ b/Engine/Plugins/Experimental/DataInterfaceGraph/Source/DataInterfaceGraphEditor/Private/DataInterfacePropertyTypeCustomization.h
@@ -0,0 +1,24 @@
+﻿// Copyright Epic Games, Inc. All Rights Reserved.
+
+#pragma once
+
+#include "IPropertyTypeCustomization.h"
+#include "PropertyEditorModule.h"
+
+namespace UE::DataInterfaceGraphEditor
+{
+
+class FDataInterfacePropertyTypeCustomization : public IPropertyTypeCustomization
+{
+private:
+	// IPropertyTypeCustomization interface
+	virtual void CustomizeHeader(TSharedRef<IPropertyHandle> PropertyHandle, FDetailWidgetRow& HeaderRow, IPropertyTypeCustomizationUtils& CustomizationUtils) override;
+	virtual void CustomizeChildren(TSharedRef<IPropertyHandle> PropertyHandle, IDetailChildrenBuilder& ChildBuilder, IPropertyTypeCustomizationUtils& CustomizationUtils) override {}
+};
+
+class FPropertyTypeIdentifier : public IPropertyTypeIdentifier
+{
+	virtual bool IsPropertyTypeCustomized(const IPropertyHandle& PropertyHandle) const override;
+};
+
+}
\ No newline at end of file
diff --git a/Engine/Plugins/Experimental/DataInterfaceGraph/Source/DataInterfaceGraphEditor/Private/SDataInterfaceGraphEditorActionMenu.cpp b/Engine/Plugins/Experimental/DataInterfaceGraph/Source/DataInterfaceGraphEditor/Private/SDataInterfaceGraphEditorActionMenu.cpp
new file mode 100644
--- /dev/null
+++ b/Engine/Plugins/Experimental/DataInterfaceGraph/Source/DataInterfaceGraphEditor/Private/SDataInterfaceGraphEditorActionMenu.cpp
@@ -0,0 +1,206 @@
+﻿// Copyright Epic Games, Inc. All Rights Reserved.
+
+#include "SDataInterfaceGraphEditorActionMenu.h"
+
+#include "DataInterfaceGraphEditorSchemaActions.h"
+#include "Framework/Application/SlateApplication.h"
+#include "IDocumentation.h"
+#include "SGraphActionMenu.h"
+#include "SSubobjectEditor.h"
+#include "DataInterfaceGraph_EdGraphSchema.h"
+#include "Widgets/Images/SImage.h"
+#include "Widgets/Layout/SBox.h"
+#include "Widgets/SBoxPanel.h"
+#include "Widgets/Text/STextBlock.h"
+#include "SGraphPalette.h"
+#include "RigVMCore/RigVMRegistry.h"
+#include "Units/RigUnit.h"
+#include "Widgets/SToolTip.h"
+
+#define LOCTEXT_NAMESPACE "DataInterfaceEditor"
+
+namespace UE::DataInterfaceGraphEditor
+{
+
+static void CollectAllDataInterfaceActions(FGraphContextMenuBuilder& MenuBuilder)
+{
+	for(const FRigVMFunction& Function : FRigVMRegistry::Get().GetFunctions())
+	{
+		UScriptStruct* Struct = Function.Struct;
+		if (!Struct->IsChildOf(FRigUnit::StaticStruct()))
+		{
+			continue;
+		}
+
+		FString CategoryMetadata, DisplayNameMetadata, MenuDescSuffixMetadata;
+		Struct->GetStringMetaDataHierarchical(FRigVMStruct::CategoryMetaName, &CategoryMetadata);
+		Struct->GetStringMetaDataHierarchical(FRigVMStruct::DisplayNameMetaName, &DisplayNameMetadata);
+		Struct->GetStringMetaDataHierarchical(FRigVMStruct::MenuDescSuffixMetaName, &MenuDescSuffixMetadata);
+		if (!MenuDescSuffixMetadata.IsEmpty())
+		{
+			MenuDescSuffixMetadata = TEXT(" ") + MenuDescSuffixMetadata;
+		}
+		FText NodeCategory = FText::FromString(CategoryMetadata);
+		FText MenuDesc = FText::FromString(DisplayNameMetadata + MenuDescSuffixMetadata);
+		FText ToolTip = Struct->GetToolTipText();
+
+		MenuBuilder.AddAction(MakeShared<FDataInterfaceGraphSchemaAction_RigUnit>(Struct, NodeCategory, MenuDesc, ToolTip));
+	};
+}
+
+SActionMenu::~SActionMenu()
+{
+	OnClosedCallback.ExecuteIfBound();
+	OnCloseReasonCallback.ExecuteIfBound(bActionExecuted, false, !DraggedFromPins.IsEmpty());
+}
+
+void SActionMenu::Construct(const FArguments& InArgs)
+{
+	Graph = InArgs._Graph;
+	DraggedFromPins = InArgs._DraggedFromPins;
+	NewNodePosition = InArgs._NewNodePosition;
+	OnClosedCallback = InArgs._OnClosedCallback;
+	bAutoExpandActionMenu = InArgs._AutoExpandActionMenu;
+	OnCloseReasonCallback = InArgs._OnCloseReason;
+
+	SBorder::Construct(SBorder::FArguments()
+		.BorderImage(FAppStyle::Get().GetBrush("Menu.Background"))
+		.Padding(5)
+		[
+			SNew(SBox)
+			.WidthOverride(400)
+			.HeightOverride(400)
+			[
+				SNew(SVerticalBox)
+				+SVerticalBox::Slot()
+				.AutoHeight()
+				.Padding(2, 2, 2, 5)
+				[
+					SNew(SHorizontalBox)
+					+SHorizontalBox::Slot()
+					.AutoWidth()
+					.VAlign(VAlign_Center)
+					[
+						SNew(STextBlock)
+						.Text(LOCTEXT("ContextText", "All Data Interface Node Classes"))
+						.Font(FAppStyle::Get().GetFontStyle("BlueprintEditor.ActionMenu.ContextDescriptionFont"))
+						.ToolTip(IDocumentation::Get()->CreateToolTip(
+							LOCTEXT("ActionMenuContextTextTooltip", "Describes the current context of the action list"),
+							nullptr,
+							TEXT("Shared/Editors/DataInterfaceEditor"),
+							TEXT("DataInterfaceActionMenuContextText")))
+						.WrapTextAt(280)
+					]
+				]
+				+SVerticalBox::Slot()
+				[
+					SAssignNew(GraphActionMenu, SGraphActionMenu)
+						.OnActionSelected(this, &SActionMenu::OnActionSelected)
+						.OnCreateWidgetForAction(SGraphActionMenu::FOnCreateWidgetForAction::CreateSP(this, &SActionMenu::OnCreateWidgetForAction))
+						.OnCollectAllActions(this, &SActionMenu::CollectAllActions)
+						.DraggedFromPins(DraggedFromPins)
+						.GraphObj(Graph)
+						.AlphaSortItems(true)
+						.bAllowPreselectedItemActivation(true)
+				]
+			]
+		]
+	);
+}
+
+void SActionMenu::CollectAllActions(FGraphActionListBuilderBase& OutAllActions)
+{
+	if (!Graph)
+	{
+		return;
+	}
+
+	FGraphContextMenuBuilder MenuBuilder(Graph);
+	if (!DraggedFromPins.IsEmpty())
+	{
+		MenuBuilder.FromPin = DraggedFromPins[0];
+	}
+
+	// Cannot call GetGraphContextActions() during serialization and GC due to its use of FindObject()
+	if(!GIsSavingPackage && !IsGarbageCollecting())
+	{
+		CollectAllDataInterfaceActions(MenuBuilder);
+	}
+
+	OutAllActions.Append(MenuBuilder);
+}
+
+TSharedRef<SEditableTextBox> SActionMenu::GetFilterTextBox()
+{
+	return GraphActionMenu->GetFilterTextBox();
+}
+
+TSharedRef<SWidget> SActionMenu::OnCreateWidgetForAction(FCreateWidgetForActionData* const InCreateData)
+{
+	check(InCreateData);
+	InCreateData->bHandleMouseButtonDown = false;
+
+	const FSlateBrush* IconBrush = nullptr;
+	FLinearColor IconColor;
+	TSharedPtr<FDataInterfaceGraphSchemaAction> Action = StaticCastSharedPtr<FDataInterfaceGraphSchemaAction>(InCreateData->Action);
+	if (Action.IsValid())
+	{
+		IconBrush = Action->GetIconBrush();
+		IconColor = Action->GetIconColor();
+	}
+
+	TSharedPtr<SHorizontalBox> WidgetBox = SNew(SHorizontalBox);
+	if (IconBrush)
+	{
+		WidgetBox->AddSlot()
+			.AutoWidth()
+			.VAlign(VAlign_Center)
+			.Padding(5, 0, 0, 0)
+			[
+				SNew(SImage)
+				.ColorAndOpacity(IconColor)
+				.Image(IconBrush)
+			];
+	}
+
+	WidgetBox->AddSlot()
+		.AutoWidth()
+		.VAlign(VAlign_Center)
+		.Padding(0, 0, 0, 0)
+		[
+			SNew(SGraphPaletteItem, InCreateData)
+		];
+
+	return WidgetBox->AsShared();
+}
+
+void SActionMenu::OnActionSelected(const TArray<TSharedPtr<FEdGraphSchemaAction>>& SelectedAction, ESelectInfo::Type InSelectionType)
+{
+	if (!Graph)
+	{
+		return;
+	}
+
+	if (InSelectionType != ESelectInfo::OnMouseClick  && InSelectionType != ESelectInfo::OnKeyPress && !SelectedAction.IsEmpty())
+	{
+		return;
+	}
+
+	for (const TSharedPtr<FEdGraphSchemaAction>& Action : SelectedAction)
+	{
+		if (Action.IsValid() && Graph)
+		{
+			if (!bActionExecuted && (Action->GetTypeId() != FEdGraphSchemaAction_Dummy::StaticGetTypeId()))
+			{
+				FSlateApplication::Get().DismissAllMenus();
+				bActionExecuted = true;
+			}
+
+			Action->PerformAction(Graph, DraggedFromPins, NewNodePosition);
+		}
+	}
+}
+
+} 
+
+#undef LOCTEXT_NAMESPACE
diff --git a/Engine/Plugins/Experimental/DataInterfaceGraph/Source/DataInterfaceGraphEditor/Private/SDataInterfaceGraphEditorActionMenu.h b/Engine/Plugins/Experimental/DataInterfaceGraph/Source/DataInterfaceGraphEditor/Private/SDataInterfaceGraphEditorActionMenu.h
new file mode 100644
--- /dev/null
+++ b/Engine/Plugins/Experimental/DataInterfaceGraph/Source/DataInterfaceGraphEditor/Private/SDataInterfaceGraphEditorActionMenu.h
@@ -0,0 +1,72 @@
+﻿// Copyright Epic Games, Inc. All Rights Reserved.
+#pragma once
+
+#include "GraphEditor.h"
+#include "EdGraph/EdGraph.h"
+#include "EdGraph/EdGraphPin.h"
+#include "EdGraph/EdGraphSchema.h"
+#include "Layout/Margin.h"
+#include "Layout/Visibility.h"
+#include "SGraphActionMenu.h"
+#include "Widgets/DeclarativeSyntaxSupport.h"
+#include "Widgets/Input/SEditableTextBox.h"
+#include "Widgets/Layout/SBorder.h"
+#include "Widgets/SWidget.h"
+#include "Widgets/Views/SExpanderArrow.h"
+
+class SEditableTextBox;
+class SGraphActionMenu;
+class UEdGraph;
+
+namespace UE::DataInterfaceGraphEditor
+{
+
+class SActionMenu : public SBorder
+{
+public:
+	/** Delegate for the OnCloseReason event which is always raised when the SActionMenu closes */
+	DECLARE_DELEGATE_ThreeParams(FClosedReason, bool /*bActionExecuted*/, bool /*bContextSensitiveChecked*/, bool /*bGraphPinContext*/);
+
+	SLATE_BEGIN_ARGS(SActionMenu)
+		: _Graph(static_cast<UEdGraph*>(nullptr))
+		, _NewNodePosition(FVector2D::ZeroVector)
+		, _AutoExpandActionMenu(false)
+	{
+	}
+		SLATE_ARGUMENT(UEdGraph*, Graph)
+		SLATE_ARGUMENT(FVector2D, NewNodePosition)
+		SLATE_ARGUMENT(TArray<UEdGraphPin*>, DraggedFromPins)
+		SLATE_ARGUMENT(SGraphEditor::FActionMenuClosed, OnClosedCallback)
+		SLATE_ARGUMENT(bool, AutoExpandActionMenu)
+		SLATE_EVENT(FClosedReason, OnCloseReason)
+	SLATE_END_ARGS()
+
+	void Construct(const FArguments& InArgs);
+
+	~SActionMenu();
+
+	TSharedRef<SEditableTextBox> GetFilterTextBox();
+
+protected:
+	void OnActionSelected(const TArray<TSharedPtr<FEdGraphSchemaAction>>& SelectedAction, ESelectInfo::Type InSelectionType);
+
+	TSharedRef<SWidget> OnCreateWidgetForAction(struct FCreateWidgetForActionData* const InCreateData);
+
+	/** Callback used to populate all actions list in SGraphActionMenu */
+	void CollectAllActions(FGraphActionListBuilderBase& OutAllActions);
+
+private:
+	UEdGraph* Graph = nullptr;
+	bool bAutoExpandActionMenu = false;
+	bool bActionExecuted = false;
+
+	TArray<UEdGraphPin*> DraggedFromPins;
+	FVector2D NewNodePosition = FVector2D::ZeroVector;
+
+	SGraphEditor::FActionMenuClosed OnClosedCallback;
+	FClosedReason OnCloseReasonCallback;
+
+	TSharedPtr<SGraphActionMenu> GraphActionMenu;
+};
+
+}
diff --git a/Engine/Plugins/Experimental/DataInterfaceGraph/Source/DataInterfaceGraphUncookedOnly/DataInterfaceGraphUncookedOnly.Build.cs b/Engine/Plugins/Experimental/DataInterfaceGraph/Source/DataInterfaceGraphUncookedOnly/DataInterfaceGraphUncookedOnly.Build.cs
new file mode 100644
--- /dev/null
+++ b/Engine/Plugins/Experimental/DataInterfaceGraph/Source/DataInterfaceGraphUncookedOnly/DataInterfaceGraphUncookedOnly.Build.cs
@@ -0,0 +1,44 @@
+// Copyright Epic Games, Inc. All Rights Reserved.
+
+namespace UnrealBuildTool.Rules
+{
+	public class DataInterfaceGraphUncookedOnly : ModuleRules
+	{
+		public DataInterfaceGraphUncookedOnly(ReadOnlyTargetRules Target) : base(Target)
+		{
+			
+			PublicDependencyModuleNames.AddRange(
+				new string[]
+				{
+					"RigVM",
+					"RigVMDeveloper",
+					"ControlRig",
+					"ControlRigDeveloper",
+				}
+			);
+
+			PrivateDependencyModuleNames.AddRange(
+				new string[]
+				{
+					"Core",
+					"CoreUObject",
+					"Engine",
+					"DataInterface",
+					"DataInterfaceGraph",
+					"BlueprintGraph",	// For K2 schema
+				}
+			);
+
+
+			if (Target.bBuildEditor)
+			{
+				PrivateDependencyModuleNames.AddRange(
+					new string[]
+					{
+						"UnrealEd",
+					}
+				);
+			}
+		}
+	}
+}
\ No newline at end of file
diff --git a/Engine/Plugins/Experimental/DataInterfaceGraph/Source/DataInterfaceGraphUncookedOnly/Internal/DataInterfaceGraph_EdGraph.h b/Engine/Plugins/Experimental/DataInterfaceGraph/Source/DataInterfaceGraphUncookedOnly/Internal/DataInterfaceGraph_EdGraph.h
new file mode 100644
--- /dev/null
+++ b/Engine/Plugins/Experimental/DataInterfaceGraph/Source/DataInterfaceGraphUncookedOnly/Internal/DataInterfaceGraph_EdGraph.h
@@ -0,0 +1,28 @@
+﻿// Copyright Epic Games, Inc. All Rights Reserved.
+
+#pragma once
+
+#include "CoreMinimal.h"
+#include "Graph/ControlRigGraph.h"
+#include "DataInterfaceGraph_EdGraph.generated.h"
+
+class UDataInterfaceGraph_EditorData;
+
+namespace UE::DataInterfaceUncookedOnly
+{
+struct FUtils;
+}
+
+UCLASS(MinimalAPI)
+class UDataInterfaceGraph_EdGraph : public UControlRigGraph
+{
+	GENERATED_BODY()
+
+	friend class UDataInterfaceGraph_EditorData;
+
+	// UControlRigGraph interface
+	virtual URigVMGraph* GetModel() const override;
+	virtual URigVMController* GetController() const override;
+	
+	void Initialize(UDataInterfaceGraph_EditorData* InEditorData);
+};
\ No newline at end of file
diff --git a/Engine/Plugins/Experimental/DataInterfaceGraph/Source/DataInterfaceGraphUncookedOnly/Internal/DataInterfaceGraph_EdGraphNode.h b/Engine/Plugins/Experimental/DataInterfaceGraph/Source/DataInterfaceGraphUncookedOnly/Internal/DataInterfaceGraph_EdGraphNode.h
new file mode 100644
--- /dev/null
+++ b/Engine/Plugins/Experimental/DataInterfaceGraph/Source/DataInterfaceGraphUncookedOnly/Internal/DataInterfaceGraph_EdGraphNode.h
@@ -0,0 +1,13 @@
+﻿// Copyright Epic Games, Inc. All Rights Reserved.
+
+#pragma once
+
+#include "CoreMinimal.h"
+#include "Graph/ControlRigGraphNode.h"
+#include "DataInterfaceGraph_EdGraphNode.generated.h"
+
+UCLASS(MinimalAPI)
+class UDataInterfaceGraph_EdGraphNode : public UControlRigGraphNode
+{
+	GENERATED_BODY()
+};
\ No newline at end of file
diff --git a/Engine/Plugins/Experimental/DataInterfaceGraph/Source/DataInterfaceGraphUncookedOnly/Internal/DataInterfaceGraph_EdGraphSchema.h b/Engine/Plugins/Experimental/DataInterfaceGraph/Source/DataInterfaceGraphUncookedOnly/Internal/DataInterfaceGraph_EdGraphSchema.h
new file mode 100644
--- /dev/null
+++ b/Engine/Plugins/Experimental/DataInterfaceGraph/Source/DataInterfaceGraphUncookedOnly/Internal/DataInterfaceGraph_EdGraphSchema.h
@@ -0,0 +1,20 @@
+﻿// Copyright Epic Games, Inc. All Rights Reserved.
+
+#pragma once
+
+#include "CoreMinimal.h"
+#include "Graph/ControlRigGraphSchema.h"
+#include "DataInterfaceGraph_EdGraphNode.h"
+#include "DataInterfaceGraph_EdGraphSchema.generated.h"
+
+UCLASS(MinimalAPI)
+class UDataInterfaceGraph_EdGraphSchema : public UControlRigGraphSchema
+{
+	GENERATED_BODY()
+
+	// UControlRigGraphSchema interface
+	virtual TSubclassOf<UControlRigGraphNode> GetGraphNodeClass() const override { return UDataInterfaceGraph_EdGraphNode::StaticClass(); }
+
+	// UEdGraphSchema interface
+	virtual void TrySetDefaultValue(UEdGraphPin& InPin, const FString& InNewDefaultValue, bool bMarkAsModified) const override;
+};
\ No newline at end of file
diff --git a/Engine/Plugins/Experimental/DataInterfaceGraph/Source/DataInterfaceGraphUncookedOnly/Internal/DataInterfaceGraph_EditorData.h b/Engine/Plugins/Experimental/DataInterfaceGraph/Source/DataInterfaceGraphUncookedOnly/Internal/DataInterfaceGraph_EditorData.h
new file mode 100644
--- /dev/null
+++ b/Engine/Plugins/Experimental/DataInterfaceGraph/Source/DataInterfaceGraphUncookedOnly/Internal/DataInterfaceGraph_EditorData.h
@@ -0,0 +1,115 @@
+﻿// Copyright Epic Games, Inc. All Rights Reserved.
+
+#pragma once
+
+#include "CoreMinimal.h"
+#include "RigVMModel/RigVMGraph.h"
+#include "DataInterfaceGraph_EdGraph.h"
+#include "DataInterfaceGraph_EditorData.generated.h"
+
+class UDataInterfaceGraph;
+enum class ERigVMGraphNotifType : uint8;
+namespace UE::DataInterfaceGraphUncookedOnly
+{
+	struct FUtils;
+}
+
+namespace UE::DataInterfaceGraphEditor
+{
+	class FGraphEditor;
+}
+
+UCLASS(Optional, MinimalAPI)
+class UDataInterfaceGraph_EditorData : public UObject, public IRigVMGraphHost, public IRigVMControllerHost
+{
+	GENERATED_BODY()
+
+	UDataInterfaceGraph_EditorData(const FObjectInitializer& ObjectInitializer);
+	
+	friend class UDataInterfaceGraphFactory;
+	friend class UDataInterfaceGraph_EdGraph;
+	friend struct UE::DataInterfaceGraphUncookedOnly::FUtils;
+	friend class UE::DataInterfaceGraphEditor::FGraphEditor;
+	friend struct FDataInterfaceGraphSchemaAction_RigUnit;
+
+	// UObject interface
+	virtual void PostLoad() override;
+	virtual bool IsEditorOnly() const override { return true; }
+	
+	// IRigVMGraphHost interface
+	virtual URigVMGraph* GetRigVMGraph(const UObject* InEditorObject) const override;
+
+	// IRigVMControllerHost interface
+	virtual URigVMController* GetRigVMController(const URigVMGraph* InRigVMGraph) const override;
+	virtual URigVMController* GetRigVMController(const UObject* InEditorObject) const override;
+	virtual URigVMController* GetOrCreateRigVMController(URigVMGraph* InRigVMGraph) override;
+	virtual URigVMController* GetOrCreateRigVMController(const UObject* InEditorObject) override;
+
+	DATAINTERFACEGRAPHUNCOOKEDONLY_API void Initialize(bool bRecompileVM);
+
+	void RecompileVM();
+	
+	void RecompileVMIfRequired();
+
+	void RequestAutoVMRecompilation();
+	
+	void HandleModifiedEvent(ERigVMGraphNotifType InNotifType, URigVMGraph* InGraph, UObject* InSubject);
+
+	URigVMGraph* GetVMGraphForEdGraph(const UEdGraph* InGraph) const;
+
+	void CreateEdGraphForCollapseNode(URigVMCollapseNode* InNode);
+	
+	UPROPERTY()
+	TObjectPtr<UDataInterfaceGraph_EdGraph> RootGraph;
+
+	UPROPERTY()
+	TObjectPtr<UDataInterfaceGraph_EdGraph> EntryPointGraph;
+	
+	UPROPERTY()
+	TObjectPtr<UDataInterfaceGraph_EdGraph> FunctionLibraryEdGraph;
+	
+	UPROPERTY()
+	TObjectPtr<URigVMGraph> RigVMGraph;
+	
+	UPROPERTY()
+	TObjectPtr<URigVMFunctionLibrary> RigVMFunctionLibrary;
+
+	UPROPERTY()
+	TObjectPtr<URigVMLibraryNode> EntryPoint;
+	
+	UPROPERTY(transient)
+	TMap<TObjectPtr<URigVMGraph>, TObjectPtr<URigVMController>> Controllers;
+
+	UPROPERTY(EditAnywhere, Category = "User Interface")
+	FRigGraphDisplaySettings RigGraphDisplaySettings;
+
+	UPROPERTY(EditAnywhere, Category = "VM")
+	FRigVMRuntimeSettings VMRuntimeSettings;
+	
+	UPROPERTY(EditAnywhere, BlueprintReadWrite, Category = "VM", meta = (AllowPrivateAccess = "true"))
+	FRigVMCompileSettings VMCompileSettings;
+
+	UPROPERTY(EditAnywhere, Category = "Python Log Settings")
+	FControlRigPythonSettings PythonLogSettings;
+
+	UPROPERTY(transient, DuplicateTransient)
+	TMap<FString, FRigVMOperand> PinToOperandMap;
+
+	UPROPERTY(transient, DuplicateTransient)
+	bool bVMRecompilationRequired = false;
+
+	UPROPERTY(transient, DuplicateTransient)
+	bool bIsCompiling = false;
+	
+	FCompilerResultsLog CompileLog;
+
+	FOnVMCompiledEvent VMCompiledEvent;
+	FRigVMGraphModifiedEvent ModifiedEvent;
+
+	bool bAutoRecompileVM = true;
+	bool bErrorsDuringCompilation = false;
+	bool bSuspendModelNotificationsForSelf = false;
+	bool bSuspendModelNotificationsForOthers = false;
+	bool bSuspendAllNotifications = false;
+	bool bCompileInDebugMode = false;
+};
\ No newline at end of file
diff --git a/Engine/Plugins/Experimental/DataInterfaceGraph/Source/DataInterfaceGraphUncookedOnly/Internal/DataInterfaceUncookedOnlyUtils.h b/Engine/Plugins/Experimental/DataInterfaceGraph/Source/DataInterfaceGraphUncookedOnly/Internal/DataInterfaceUncookedOnlyUtils.h
new file mode 100644
--- /dev/null
+++ b/Engine/Plugins/Experimental/DataInterfaceGraph/Source/DataInterfaceGraphUncookedOnly/Internal/DataInterfaceUncookedOnlyUtils.h
@@ -0,0 +1,30 @@
+﻿// Copyright Epic Games, Inc. All Rights Reserved.
+
+#pragma once
+
+#include "CoreMinimal.h"
+
+class UDataInterfaceGraph;
+class UDataInterfaceGraph_EditorData;
+class UDataInterfaceGraph_EdGraph;
+class URigVMController;
+class URigVMGraph;
+
+namespace UE::DataInterfaceGraphEditor
+{
+	class FGraphEditor;
+}
+
+namespace UE::DataInterfaceGraphUncookedOnly
+{
+
+struct DATAINTERFACEGRAPHUNCOOKEDONLY_API FUtils
+{
+	static void Compile(UDataInterfaceGraph* InGraph);
+	
+	static UDataInterfaceGraph_EditorData* GetEditorData(const UDataInterfaceGraph* InDataInterfaceGraph);
+	
+	static void RecreateVM(UDataInterfaceGraph* InGraph);
+};
+
+}
\ No newline at end of file
diff --git a/Engine/Plugins/Experimental/DataInterfaceGraph/Source/DataInterfaceGraphUncookedOnly/Private/DataInterfaceGraphUncookedOnlyModule.cpp b/Engine/Plugins/Experimental/DataInterfaceGraph/Source/DataInterfaceGraphUncookedOnly/Private/DataInterfaceGraphUncookedOnlyModule.cpp
new file mode 100644
--- /dev/null
+++ b/Engine/Plugins/Experimental/DataInterfaceGraph/Source/DataInterfaceGraphUncookedOnly/Private/DataInterfaceGraphUncookedOnlyModule.cpp
@@ -0,0 +1,5 @@
+// Copyright Epic Games, Inc. All Rights Reserved.
+
+#include "Modules/ModuleManager.h"
+
+IMPLEMENT_MODULE(FDefaultModuleImpl, DataInterfaceUncookedOnly);
\ No newline at end of file
diff --git a/Engine/Plugins/Experimental/DataInterfaceGraph/Source/DataInterfaceGraphUncookedOnly/Private/DataInterfaceGraphUncookedOnlyUtils.cpp b/Engine/Plugins/Experimental/DataInterfaceGraph/Source/DataInterfaceGraphUncookedOnly/Private/DataInterfaceGraphUncookedOnlyUtils.cpp
new file mode 100644
--- /dev/null
+++ b/Engine/Plugins/Experimental/DataInterfaceGraph/Source/DataInterfaceGraphUncookedOnly/Private/DataInterfaceGraphUncookedOnlyUtils.cpp
@@ -0,0 +1,104 @@
+﻿// Copyright Epic Games, Inc. All Rights Reserved.
+
+#include "DataInterfaceUncookedOnlyUtils.h"
+#include "DataInterfaceGraph.h"
+#include "DataInterfaceGraph_EditorData.h"
+#include "DataInterfaceGraph_EdGraph.h"
+#include "RigVMCompiler/RigVMCompiler.h"
+#include "RigVMCore/RigVM.h"
+#include "Units/RigUnitContext.h"
+
+namespace UE::DataInterfaceGraphUncookedOnly
+{
+
+void FUtils::Compile(UDataInterfaceGraph* InGraph)
+{
+	check(InGraph);
+	
+	UDataInterfaceGraph_EditorData* EditorData = GetEditorData(InGraph);
+	
+	if(EditorData->bIsCompiling)
+	{
+		return;
+	}
+	
+	TGuardValue<bool> CompilingGuard(EditorData->bIsCompiling, true);
+	
+	EditorData->bErrorsDuringCompilation = false;
+
+	EditorData->RigGraphDisplaySettings.MinMicroSeconds = EditorData->RigGraphDisplaySettings.LastMinMicroSeconds = DBL_MAX;
+	EditorData->RigGraphDisplaySettings.MaxMicroSeconds = EditorData->RigGraphDisplaySettings.LastMaxMicroSeconds = (double)INDEX_NONE;
+	
+	TGuardValue<bool> ReentrantGuardSelf(EditorData->bSuspendModelNotificationsForSelf, true);
+	TGuardValue<bool> ReentrantGuardOthers(EditorData->bSuspendModelNotificationsForOthers, true);
+
+	RecreateVM(InGraph);
+
+	InGraph->VMRuntimeSettings = EditorData->VMRuntimeSettings;
+
+	FRigNameCache TempNameCache;
+	FRigUnitContext InitContext;
+	InitContext.State = EControlRigState::Init;
+	InitContext.NameCache = &TempNameCache;
+
+	FRigUnitContext UpdateContext = InitContext;
+	UpdateContext.State = EControlRigState::Update;
+
+	void* InitContextPtr = &InitContext;
+	void* UpdateContextPtr = &UpdateContext;
+
+	TArray<FRigVMUserDataArray> UserData;
+	UserData.Add(FRigVMUserDataArray(&InitContextPtr, 1));
+	UserData.Add(FRigVMUserDataArray(&UpdateContextPtr, 1));
+
+	EditorData->CompileLog.Messages.Reset();
+	EditorData->CompileLog.NumErrors = EditorData->CompileLog.NumWarnings = 0;
+
+	URigVMCompiler* Compiler = URigVMCompiler::StaticClass()->GetDefaultObject<URigVMCompiler>();
+	Compiler->Settings = (EditorData->bCompileInDebugMode) ? FRigVMCompileSettings::Fast() : EditorData->VMCompileSettings;
+	URigVMController* RootController = EditorData->GetOrCreateRigVMController(EditorData->RigVMGraph);
+	Compiler->Compile(EditorData->RigVMGraph, RootController, InGraph->RigVM, InGraph->GetRigVMExternalVariables(), UserData, &EditorData->PinToOperandMap);
+
+	if (EditorData->bErrorsDuringCompilation)
+	{
+		if(Compiler->Settings.SurpressErrors)
+		{
+			Compiler->Settings.Reportf(EMessageSeverity::Info, InGraph,TEXT("Compilation Errors may be suppressed for Data Interface Graph: %s. See VM Compile Settings for more Details"), *InGraph->GetName());
+		}
+	}
+
+	EditorData->bVMRecompilationRequired = false;
+	if(InGraph->RigVM)
+	{
+		EditorData->VMCompiledEvent.Broadcast(InGraph, InGraph->RigVM);
+	}
+
+#if WITH_EDITOR
+//	RefreshBreakpoints(EditorData);
+#endif
+}
+
+void FUtils::RecreateVM(UDataInterfaceGraph* InGraph)
+{
+	InGraph->RigVM = NewObject<URigVM>(InGraph, TEXT("VM"), RF_NoFlags);
+	
+	// Cooked platforms will load these pointers from disk
+	if (!FPlatformProperties::RequiresCookedData())
+	{
+		// We dont support ERigVMMemoryType::Work memory as we dont operate on an instance
+	//	InGraph->RigVM->GetMemoryByType(ERigVMMemoryType::Work, true);
+		InGraph->RigVM->GetMemoryByType(ERigVMMemoryType::Literal, true);
+		InGraph->RigVM->GetMemoryByType(ERigVMMemoryType::Debug, true);
+	}
+
+	InGraph->RigVM->Reset();
+}
+
+UDataInterfaceGraph_EditorData* FUtils::GetEditorData(const UDataInterfaceGraph* InDataInterfaceGraph)
+{
+	check(InDataInterfaceGraph);
+	
+	return CastChecked<UDataInterfaceGraph_EditorData>(InDataInterfaceGraph->EditorData);
+}
+
+}
diff --git a/Engine/Plugins/Experimental/DataInterfaceGraph/Source/DataInterfaceGraphUncookedOnly/Private/DataInterfaceGraph_EdGraph.cpp b/Engine/Plugins/Experimental/DataInterfaceGraph/Source/DataInterfaceGraphUncookedOnly/Private/DataInterfaceGraph_EdGraph.cpp
new file mode 100644
--- /dev/null
+++ b/Engine/Plugins/Experimental/DataInterfaceGraph/Source/DataInterfaceGraphUncookedOnly/Private/DataInterfaceGraph_EdGraph.cpp
@@ -0,0 +1,31 @@
+﻿// Copyright Epic Games, Inc. All Rights Reserved.
+
+#include "DataInterfaceGraph_EdGraph.h"
+#include "DataInterfaceGraph_EditorData.h"
+#include "DataInterfaceUncookedOnlyUtils.h"
+
+void UDataInterfaceGraph_EdGraph::Initialize(UDataInterfaceGraph_EditorData* InEditorData)
+{
+	InEditorData->ModifiedEvent.RemoveAll(this);
+	InEditorData->ModifiedEvent.AddUObject(this, &UDataInterfaceGraph_EdGraph::HandleModifiedEvent);
+	InEditorData->VMCompiledEvent.RemoveAll(this);
+	InEditorData->VMCompiledEvent.AddUObject(this, &UDataInterfaceGraph_EdGraph::HandleVMCompiledEvent);
+}
+
+URigVMGraph* UDataInterfaceGraph_EdGraph::GetModel() const
+{
+	if (const UDataInterfaceGraph_EditorData* EditorData = GetTypedOuter<UDataInterfaceGraph_EditorData>())
+	{
+		return EditorData->GetRigVMGraph(this);
+	}
+	return nullptr;
+}
+
+URigVMController* UDataInterfaceGraph_EdGraph::GetController() const
+{
+	if (const UDataInterfaceGraph_EditorData* EditorData = GetTypedOuter<UDataInterfaceGraph_EditorData>())
+	{
+		return EditorData->GetRigVMController(this);
+	}
+	return nullptr;
+}
\ No newline at end of file
diff --git a/Engine/Plugins/Experimental/DataInterfaceGraph/Source/DataInterfaceGraphUncookedOnly/Private/DataInterfaceGraph_EdGraphSchema.cpp b/Engine/Plugins/Experimental/DataInterfaceGraph/Source/DataInterfaceGraphUncookedOnly/Private/DataInterfaceGraph_EdGraphSchema.cpp
new file mode 100644
--- /dev/null
+++ b/Engine/Plugins/Experimental/DataInterfaceGraph/Source/DataInterfaceGraphUncookedOnly/Private/DataInterfaceGraph_EdGraphSchema.cpp
@@ -0,0 +1,42 @@
+﻿// Copyright Epic Games, Inc. All Rights Reserved.
+
+#include "DataInterfaceGraph_EdGraphSchema.h"
+#include "Editor.h"
+
+void UDataInterfaceGraph_EdGraphSchema::TrySetDefaultValue(UEdGraphPin& InPin, const FString& InNewDefaultValue, bool bMarkAsModified) const
+{
+#if WITH_EDITOR
+	if (GEditor)
+	{
+		GEditor->CancelTransaction(0);
+	}
+#endif
+
+	FString UseDefaultValue;
+	TObjectPtr<UObject> UseDefaultObject = nullptr;
+	FText UseDefaultText;
+
+	GetDefault<UEdGraphSchema_K2>()->GetPinDefaultValuesFromString(InPin.PinType, InPin.GetOwningNodeUnchecked(), InNewDefaultValue, UseDefaultValue, UseDefaultObject, UseDefaultText, /*bPreserveTextIdentity*/false);
+
+	// Check the default value and make it an error if it's bogus
+	if (GetDefault<UEdGraphSchema_K2>()->DefaultValueSimpleValidation(InPin.PinType, InPin.PinName, UseDefaultValue, UseDefaultObject, UseDefaultText))
+	{
+		InPin.DefaultObject = UseDefaultObject;
+		InPin.DefaultValue = UseDefaultValue;
+		InPin.DefaultTextValue = UseDefaultText;
+
+		UEdGraphNode* Node = InPin.GetOwningNode();
+		Node->PinDefaultValueChanged(&InPin);
+
+		// If the default value is manually set then treat it as if the value was reset to default and remove the orphaned InPin
+		if (InPin.bOrphanedPin && InPin.DoesDefaultValueMatchAutogenerated())
+		{
+			Node->PinConnectionListChanged(&InPin);
+		}
+
+		if (bMarkAsModified)
+		{
+			Node->MarkPackageDirty();
+		}
+	}
+}
diff --git a/Engine/Plugins/Experimental/DataInterfaceGraph/Source/DataInterfaceGraphUncookedOnly/Private/DataInterfaceGraph_EditorData.cpp b/Engine/Plugins/Experimental/DataInterfaceGraph/Source/DataInterfaceGraphUncookedOnly/Private/DataInterfaceGraph_EditorData.cpp
new file mode 100644
--- /dev/null
+++ b/Engine/Plugins/Experimental/DataInterfaceGraph/Source/DataInterfaceGraphUncookedOnly/Private/DataInterfaceGraph_EditorData.cpp
@@ -0,0 +1,295 @@
+﻿// Copyright Epic Games, Inc. All Rights Reserved.
+
+#include "DataInterfaceGraph_EditorData.h"
+
+#include "ControlRigDefines.h"
+#include "DataInterfaceGraph.h"
+#include "DataInterfaceGraph_EdGraphSchema.h"
+#include "DataInterfaceUncookedOnlyUtils.h"
+#include "DataInterfaceExecuteContext.h"
+#include "Rigs/RigHierarchyPose.h"
+#include "RigVMModel/RigVMFunctionLibrary.h"
+#include "RigVMModel/RigVMNotifications.h"
+#include "RigVMModel/Nodes/RigVMCollapseNode.h"
+#include "Curves/CurveFloat.h"
+
+UDataInterfaceGraph_EditorData::UDataInterfaceGraph_EditorData(const FObjectInitializer& ObjectInitializer)
+{
+	RigVMGraph = ObjectInitializer.CreateDefaultSubobject<URigVMGraph>(this, TEXT("RigVMGraph"));
+	RigVMFunctionLibrary = ObjectInitializer.CreateDefaultSubobject<URigVMFunctionLibrary>(this, TEXT("RigVMFunctionLibrary"));
+	RigVMGraph->SetDefaultFunctionLibrary(RigVMFunctionLibrary);
+	
+	auto MakeEdGraph = [this, &ObjectInitializer](FName InName) -> UDataInterfaceGraph_EdGraph*
+	{
+		UDataInterfaceGraph_EdGraph* EdGraph = ObjectInitializer.CreateDefaultSubobject<UDataInterfaceGraph_EdGraph>(this, InName);
+		EdGraph->Schema = UDataInterfaceGraph_EdGraphSchema::StaticClass();
+		EdGraph->bAllowRenaming = false;
+		EdGraph->bEditable = false;
+		EdGraph->bAllowDeletion = false;
+		EdGraph->bIsFunctionDefinition = false;
+		EdGraph->Initialize(this);
+
+		return EdGraph;
+	};
+
+	RootGraph = MakeEdGraph(TEXT("RootEdGraph"));
+	FunctionLibraryEdGraph = MakeEdGraph(TEXT("RigVMFunctionLibraryEdGraph"));
+}
+
+void UDataInterfaceGraph_EditorData::Initialize(bool bRecompileVM)
+{
+	UDataInterfaceGraph* DataInterfaceGraph = GetTypedOuter<UDataInterfaceGraph>();
+	
+	if (Controllers.Num() == 0)
+	{
+		GetOrCreateRigVMController(RigVMGraph);
+
+		// Init function library controllers
+		GetOrCreateRigVMController(RigVMFunctionLibrary);
+		for(URigVMLibraryNode* LibraryNode : RigVMFunctionLibrary->GetFunctions())
+		{
+			GetOrCreateRigVMController(LibraryNode->GetContainedGraph());
+		}
+		
+		if(bRecompileVM)
+		{
+			UE::DataInterfaceGraphUncookedOnly::FUtils::Compile(DataInterfaceGraph);
+		}
+	}
+
+	RootGraph->Initialize(this);
+	FunctionLibraryEdGraph->Initialize(this);
+	if(EntryPointGraph)
+	{
+		EntryPointGraph->Initialize(this);
+	}
+}
+
+void UDataInterfaceGraph_EditorData::PostLoad()
+{
+	Super::PostLoad();
+	
+	Initialize(/*bRecompileVM*/false);
+}
+
+void UDataInterfaceGraph_EditorData::RecompileVM()
+{
+	UE::DataInterfaceGraphUncookedOnly::FUtils::Compile(GetTypedOuter<UDataInterfaceGraph>());
+}
+
+void UDataInterfaceGraph_EditorData::RecompileVMIfRequired()
+{
+	if (bVMRecompilationRequired)
+	{
+		RecompileVM();
+	}
+}
+
+void UDataInterfaceGraph_EditorData::RequestAutoVMRecompilation()
+{
+	bVMRecompilationRequired = true;
+	if (bAutoRecompileVM)
+	{
+		RecompileVMIfRequired();
+	}
+}
+
+void UDataInterfaceGraph_EditorData::HandleModifiedEvent(ERigVMGraphNotifType InNotifType, URigVMGraph* InGraph, UObject* InSubject)
+{
+	bool bNotifForOthersPending = true;
+
+	switch(InNotifType)
+	{
+	case ERigVMGraphNotifType::NodeAdded:
+		{
+			if (URigVMCollapseNode* CollapseNode = Cast<URigVMCollapseNode>(InSubject))
+			{
+				CreateEdGraphForCollapseNode(CollapseNode);
+				break;
+			}
+		}
+		// Fall through to the next case
+	case ERigVMGraphNotifType::LinkAdded:
+	case ERigVMGraphNotifType::LinkRemoved:
+	case ERigVMGraphNotifType::PinArraySizeChanged:
+	case ERigVMGraphNotifType::PinDirectionChanged:
+	case ERigVMGraphNotifType::PinDefaultValueChanged:
+		{
+			RequestAutoVMRecompilation();
+			break;
+		}
+	}
+	
+	// if the notification still has to be sent...
+	if (bNotifForOthersPending && !bSuspendModelNotificationsForOthers)
+	{
+		if (ModifiedEvent.IsBound())
+		{
+			ModifiedEvent.Broadcast(InNotifType, InGraph, InSubject);
+		}
+	}
+}
+
+URigVMGraph* UDataInterfaceGraph_EditorData::GetRigVMGraph(const UObject* InEditorObject) const
+{
+	return GetVMGraphForEdGraph(CastChecked<UEdGraph>(InEditorObject));
+}
+
+URigVMGraph* UDataInterfaceGraph_EditorData::GetVMGraphForEdGraph(const UEdGraph* InGraph) const
+{
+	if (InGraph == RootGraph)
+	{
+		return RigVMGraph;
+	}
+	else
+	{
+		const UDataInterfaceGraph_EdGraph* Graph = Cast<UDataInterfaceGraph_EdGraph>(InGraph);
+		check(Graph);
+
+		if (Graph->bIsFunctionDefinition)
+		{
+			if (URigVMLibraryNode* LibraryNode = RigVMFunctionLibrary->FindFunction(*Graph->ModelNodePath))
+			{
+				return LibraryNode->GetContainedGraph();
+			}
+		}
+	}
+
+	return nullptr;
+}
+
+URigVMController* UDataInterfaceGraph_EditorData::GetRigVMController(const URigVMGraph* InRigVMGraph) const
+{
+	TObjectPtr<URigVMController> const* ControllerPtr = Controllers.Find(InRigVMGraph);
+	if (ControllerPtr)
+	{
+		return *ControllerPtr;
+	}
+	return nullptr;
+}
+
+URigVMController* UDataInterfaceGraph_EditorData::GetRigVMController(const UObject* InEditorObject) const
+{
+	return GetRigVMController(GetVMGraphForEdGraph(CastChecked<UEdGraph>(InEditorObject)));
+}
+
+URigVMController* UDataInterfaceGraph_EditorData::GetOrCreateRigVMController(URigVMGraph* InRigVMGraph)
+{
+	if (URigVMController* ExistingController = GetRigVMController(InRigVMGraph))
+	{
+		return ExistingController;
+	}
+
+	URigVMController* Controller = NewObject<URigVMController>(this);
+	Controller->SetExecuteContextStruct(FRigVMExecuteContext::StaticStruct());
+	Controller->SetGraph(InRigVMGraph);
+	Controller->OnModified().AddUObject(this, &UDataInterfaceGraph_EditorData::HandleModifiedEvent);
+
+	Controller->UnfoldStructDelegate.BindLambda([](const UStruct* InStruct) -> bool
+	{
+		if (InStruct == TBaseStructure<FQuat>::Get())
+		{
+			return false;
+		}
+		if (InStruct == FRuntimeFloatCurve::StaticStruct())
+		{
+			return false;
+		}
+		if (InStruct == FRigPose::StaticStruct())
+		{
+			return false;
+		}
+		return true;
+	});
+
+	TWeakObjectPtr<UDataInterfaceGraph_EditorData> WeakThis(this);
+
+	// this delegate is used by the controller to determine variable validity
+	// during a bind process. the controller itself doesn't own the variables,
+	// so we need a delegate to request them from the owning blueprint
+	Controller->GetExternalVariablesDelegate.BindLambda([](URigVMGraph* InGraph) -> TArray<FRigVMExternalVariable>
+	{
+		if (InGraph)
+		{
+			if(UDataInterfaceGraph_EditorData* EditorData = InGraph->GetTypedOuter<UDataInterfaceGraph_EditorData>())
+			{
+				if (UDataInterfaceGraph* Graph = EditorData->GetTypedOuter<UDataInterfaceGraph>())
+				{
+					return Graph->GetRigVMExternalVariables();
+				}
+			}
+		}
+		return TArray<FRigVMExternalVariable>();
+	});
+
+	// this delegate is used by the controller to retrieve the current bytecode of the VM
+	Controller->GetCurrentByteCodeDelegate.BindLambda([WeakThis]() -> const FRigVMByteCode*
+	{
+		if (WeakThis.IsValid())
+		{
+			if(UDataInterfaceGraph* Graph = WeakThis->GetTypedOuter<UDataInterfaceGraph>())
+			{
+				if (Graph->RigVM)
+				{
+					return &Graph->RigVM->GetByteCode();
+				}
+			}
+		}
+		return nullptr;
+
+	});
+
+	Controller->IsFunctionAvailableDelegate.BindLambda([](URigVMLibraryNode* InFunction) -> bool
+	{
+		return true;	// @TODO: should only allow main entry point function here
+	});
+
+	Controller->IsDependencyCyclicDelegate.BindLambda([](UObject* InDependentObject, UObject* InDependencyObject) -> bool
+	{
+		return false;
+	});
+
+#if WITH_EDITOR
+	Controller->SetupDefaultUnitNodeDelegates(TDelegate<FName(FRigVMExternalVariable, FString)>::CreateLambda(
+		[](FRigVMExternalVariable InVariableToCreate, FString InDefaultValue) -> FName
+		{
+			return NAME_None;
+		}
+	));
+#endif
+
+	Controller->RemoveStaleNodes();
+	Controllers.Add(InRigVMGraph, Controller);
+	return Controller;
+}
+
+URigVMController* UDataInterfaceGraph_EditorData::GetOrCreateRigVMController(const UObject* InEditorObject)
+{
+	return GetOrCreateRigVMController(GetVMGraphForEdGraph(CastChecked<UEdGraph>(InEditorObject)));
+}
+
+void UDataInterfaceGraph_EditorData::CreateEdGraphForCollapseNode(URigVMCollapseNode* InNode)
+{
+	if (InNode->GetGraph()->IsA<URigVMFunctionLibrary>())
+	{
+		if (URigVMGraph* ContainedGraph = InNode->GetContainedGraph())
+		{
+			if (EntryPointGraph == nullptr)
+			{
+				// create a sub graph
+				UDataInterfaceGraph_EdGraph* RigFunctionGraph = NewObject<UDataInterfaceGraph_EdGraph>(this, *InNode->GetName(), RF_Transactional);
+				RigFunctionGraph->Schema = UDataInterfaceGraph_EdGraphSchema::StaticClass();
+				RigFunctionGraph->bAllowRenaming = true;
+				RigFunctionGraph->bEditable = true;
+				RigFunctionGraph->bAllowDeletion = true;
+				RigFunctionGraph->ModelNodePath = ContainedGraph->GetNodePath();
+				RigFunctionGraph->bIsFunctionDefinition = true;
+
+				EntryPointGraph = RigFunctionGraph;
+				RigFunctionGraph->Initialize(this);
+
+				GetOrCreateRigVMController(ContainedGraph)->ResendAllNotifications();
+			}
+		}
+	}
+}
\ No newline at end of file
