Fix CIS issues in DataInterface plugins

diff --git a/Engine/Plugins/Experimental/DataInterface/Source/DataInterface/Public/IDataInterfaceParameters.h b/Engine/Plugins/Experimental/DataInterface/Source/DataInterface/Public/IDataInterfaceParameters.h
--- a/Engine/Plugins/Experimental/DataInterface/Source/DataInterface/Public/IDataInterfaceParameters.h
+++ b/Engine/Plugins/Experimental/DataInterface/Source/DataInterface/Public/IDataInterfaceParameters.h
@@ -9,6 +9,7 @@
 namespace UE::DataInterface
 {
 struct FParam;
+struct FContext;
 }
 
 UINTERFACE(BlueprintType, meta = (CannotImplementInterfaceInBlueprint))
diff --git a/Engine/Plugins/Experimental/DataInterfaceGraph/Source/DataInterfaceGraph/Private/DataInterface_Float.h b/Engine/Plugins/Experimental/DataInterfaceGraph/Source/DataInterfaceGraph/Private/DataInterface_Float.h
--- a/Engine/Plugins/Experimental/DataInterfaceGraph/Source/DataInterfaceGraph/Private/DataInterface_Float.h
+++ b/Engine/Plugins/Experimental/DataInterfaceGraph/Source/DataInterfaceGraph/Private/DataInterface_Float.h
@@ -19,7 +19,11 @@ class UDataInterfaceFloat : public UObject, public IDataInterface
 		return Name;
 	}
 
-	virtual bool GetDataImpl(const UE::DataInterface::FContext& Context) const override PURE_VIRTUAL(UDataInterfaceFloat::GetDataImpl, return false;)
+	virtual bool GetDataImpl(const UE::DataInterface::FContext& Context) const override
+	{
+		checkf(false, TEXT("UDataInterfaceFloat::GetDataImpl must be overridden"));
+		return false;
+	}
 };
 
 UCLASS()
