DataInterfaceGraph: Fix member initialization

diff --git a/Engine/Plugins/Experimental/DataInterfaceGraph/Source/DataInterfaceGraph/Internal/RigUnit_DataInterfaceEndExecution.h b/Engine/Plugins/Experimental/DataInterfaceGraph/Source/DataInterfaceGraph/Internal/RigUnit_DataInterfaceEndExecution.h
--- a/Engine/Plugins/Experimental/DataInterfaceGraph/Source/DataInterfaceGraph/Internal/RigUnit_DataInterfaceEndExecution.h
+++ b/Engine/Plugins/Experimental/DataInterfaceGraph/Source/DataInterfaceGraph/Internal/RigUnit_DataInterfaceEndExecution.h
@@ -35,6 +35,10 @@ struct DATAINTERFACEGRAPH_API FRigUnit_DataInterfaceEndExecution_Bool : public F
 {
 	GENERATED_BODY()
 	
+	FRigUnit_DataInterfaceEndExecution_Bool()
+		: Result(false)
+	{}
+	
 	RIGVM_METHOD()
 	virtual void Execute(const FRigUnitContext& Context) override;
 	
@@ -47,7 +51,11 @@ USTRUCT(meta=(DisplayName="End Execute Data Interface Float"))
 struct DATAINTERFACEGRAPH_API FRigUnit_DataInterfaceEndExecution_Float : public FRigUnit_DataInterfaceEndExecution
 {
 	GENERATED_BODY()
-	
+
+	FRigUnit_DataInterfaceEndExecution_Float()
+	: Result(0.f)
+	{}
+
 	RIGVM_METHOD()
 	virtual void Execute(const FRigUnitContext& Context) override;
 
diff --git a/Engine/Plugins/Experimental/DataInterfaceGraph/Source/DataInterfaceGraph/Private/RigUnit_DataInterfaceParameter.h b/Engine/Plugins/Experimental/DataInterfaceGraph/Source/DataInterfaceGraph/Private/RigUnit_DataInterfaceParameter.h
--- a/Engine/Plugins/Experimental/DataInterfaceGraph/Source/DataInterfaceGraph/Private/RigUnit_DataInterfaceParameter.h
+++ b/Engine/Plugins/Experimental/DataInterfaceGraph/Source/DataInterfaceGraph/Private/RigUnit_DataInterfaceParameter.h
@@ -55,6 +55,10 @@ struct FDataInterfaceParameter
 {
 	GENERATED_BODY()
 
+	FDataInterfaceParameter()
+	: Name(NAME_None)
+	{}
+
 	UPROPERTY(EditAnywhere, Category = "Data Interface")
 	FName Name;
 };
@@ -63,6 +67,11 @@ USTRUCT(BlueprintType)
 struct FDataInterfaceParameter_Float
 {
 	GENERATED_BODY()
+	
+	FDataInterfaceParameter_Float()
+		: Name(NAME_None)
+		, Value(0.f)
+	{}
 
 	UPROPERTY(EditAnywhere, Category = "Data Interface")
 	FName Name;
@@ -76,6 +85,11 @@ struct FDataInterfaceParameter_Int
 {
 	GENERATED_BODY()
 
+	FDataInterfaceParameter_Int()
+		: Name(NAME_None)
+		, Value(0)
+	{}
+
 	UPROPERTY(EditAnywhere, Category = "Data Interface")
 	FName Name;
 
@@ -88,6 +102,11 @@ struct FDataInterfaceParameter_Bool
 {
 	GENERATED_BODY()
 
+	FDataInterfaceParameter_Bool()
+		: Name(NAME_None)
+		, bValue(false)
+	{}
+
 	UPROPERTY(EditAnywhere, Category = "Data Interface")
 	FName Name;
 
@@ -100,6 +119,11 @@ struct FDataInterfaceParameters1
 {
 	GENERATED_BODY()
 
+	FDataInterfaceParameters1()
+		: Param0()
+		, FloatParam(0.f)
+	{}
+
 	UPROPERTY(EditAnywhere, Category = "Data Interface")
 	FDataInterfaceParameter_Int Param0;
 
@@ -170,6 +194,12 @@ USTRUCT(meta=(DisplayName="Float Operator"))
 struct FRigUnit_FloatOperator : public FRigUnit
 {
 	GENERATED_BODY()
+	
+	FRigUnit_FloatOperator()
+		: ParamA(0.f)
+		, ParamB(0.f)
+		, Result(0.f)
+	{}
 
 	RIGVM_METHOD()
 	virtual void Execute(const FRigUnitContext& Context) override;
