fixed module name duplication

diff --git a/Engine/Plugins/Experimental/DataInterfaceGraph/Source/DataInterfaceGraph/Private/DataInterfaceGraphModule.cpp b/Engine/Plugins/Experimental/DataInterfaceGraph/Source/DataInterfaceGraph/Private/DataInterfaceGraphModule.cpp
--- a/Engine/Plugins/Experimental/DataInterfaceGraph/Source/DataInterfaceGraph/Private/DataInterfaceGraphModule.cpp
+++ b/Engine/Plugins/Experimental/DataInterfaceGraph/Source/DataInterfaceGraph/Private/DataInterfaceGraphModule.cpp
@@ -2,4 +2,4 @@
 
 #include "Modules/ModuleManager.h"
 
-IMPLEMENT_MODULE(FDefaultModuleImpl, DataInterface);
\ No newline at end of file
+IMPLEMENT_MODULE(FDefaultModuleImpl, DataInterfaceGraph);
\ No newline at end of file
