Fix static analysis warning (unnecessary duplicated check)

diff --git a/Engine/Plugins/Experimental/DataInterfaceGraph/Source/DataInterfaceGraphUncookedOnly/Private/DataInterfaceGraph_EditorData.cpp b/Engine/Plugins/Experimental/DataInterfaceGraph/Source/DataInterfaceGraphUncookedOnly/Private/DataInterfaceGraph_EditorData.cpp
--- a/Engine/Plugins/Experimental/DataInterfaceGraph/Source/DataInterfaceGraphUncookedOnly/Private/DataInterfaceGraph_EditorData.cpp
+++ b/Engine/Plugins/Experimental/DataInterfaceGraph/Source/DataInterfaceGraphUncookedOnly/Private/DataInterfaceGraph_EditorData.cpp
@@ -48,13 +48,10 @@ void UDataInterfaceGraph_EditorData::Serialize(FArchive& Ar)
 
 	if(Ar.IsLoading())
 	{
-		if(Ar.IsLoading())
+		if(RigVMGraph_DEPRECATED || RigVMFunctionLibrary_DEPRECATED)
 		{
-			if(RigVMGraph_DEPRECATED || RigVMFunctionLibrary_DEPRECATED)
-			{
-				TGuardValue<bool> DisableClientNotifs(RigVMClient.bSuspendNotifications, true);
-				RigVMClient.SetFromDeprecatedData(RigVMGraph_DEPRECATED, RigVMFunctionLibrary_DEPRECATED);
-			}
+			TGuardValue<bool> DisableClientNotifs(RigVMClient.bSuspendNotifications, true);
+			RigVMClient.SetFromDeprecatedData(RigVMGraph_DEPRECATED, RigVMFunctionLibrary_DEPRECATED);
 		}
 	}
 }
