Remove experimental class UDataInterfaceGraph_EditorData from being optional until revisited

diff --git a/Engine/Plugins/Experimental/DataInterfaceGraph/Source/DataInterfaceGraphUncookedOnly/Internal/DataInterfaceGraph_EditorData.h b/Engine/Plugins/Experimental/DataInterfaceGraph/Source/DataInterfaceGraphUncookedOnly/Internal/DataInterfaceGraph_EditorData.h
--- a/Engine/Plugins/Experimental/DataInterfaceGraph/Source/DataInterfaceGraphUncookedOnly/Internal/DataInterfaceGraph_EditorData.h
+++ b/Engine/Plugins/Experimental/DataInterfaceGraph/Source/DataInterfaceGraphUncookedOnly/Internal/DataInterfaceGraph_EditorData.h
@@ -19,7 +19,7 @@ namespace UE::DataInterfaceGraphEditor
 	class FGraphEditor;
 }
 
-UCLASS(Optional, MinimalAPI)
+UCLASS(MinimalAPI)
 class UDataInterfaceGraph_EditorData : public UObject, public IRigVMClientHost
 {
 	GENERATED_BODY()
