Fix for DataInterfaceGraph test, so it compiles and saves

diff --git a/Engine/Plugins/Experimental/DataInterfaceGraph/Source/DataInterfaceGraph/Internal/RigUnit_DataInterfaceBeginExecution.h b/Engine/Plugins/Experimental/DataInterfaceGraph/Source/DataInterfaceGraph/Internal/RigUnit_DataInterfaceBeginExecution.h
--- a/Engine/Plugins/Experimental/DataInterfaceGraph/Source/DataInterfaceGraph/Internal/RigUnit_DataInterfaceBeginExecution.h
+++ b/Engine/Plugins/Experimental/DataInterfaceGraph/Source/DataInterfaceGraph/Internal/RigUnit_DataInterfaceBeginExecution.h
@@ -18,6 +18,7 @@ struct DATAINTERFACEGRAPH_API FRigUnit_DataInterfaceBeginExecution : public FRig
 	virtual void Execute(const FRigUnitContext& Context) override;
 
 	virtual FName GetEventName() const override { return EventName; }
+	virtual bool CanOnlyExistOnce() const override { return true; }
 
 	// The execution result
 	UPROPERTY(EditAnywhere, DisplayName = "Execute", Category = "BeginExecution", meta = (Output))
diff --git a/Engine/Plugins/Experimental/DataInterfaceGraph/Source/DataInterfaceGraphEditor/Private/DataInterfaceGraphFactory.cpp b/Engine/Plugins/Experimental/DataInterfaceGraph/Source/DataInterfaceGraphEditor/Private/DataInterfaceGraphFactory.cpp
--- a/Engine/Plugins/Experimental/DataInterfaceGraph/Source/DataInterfaceGraphEditor/Private/DataInterfaceGraphFactory.cpp
+++ b/Engine/Plugins/Experimental/DataInterfaceGraph/Source/DataInterfaceGraphEditor/Private/DataInterfaceGraphFactory.cpp
@@ -34,44 +34,48 @@ UObject* UDataInterfaceGraphFactory::FactoryCreateNew(UClass* Class, UObject* In
 	EditorData->Initialize(/*bRecompileVM*/false);
 
 	// Add initial execution unit
-	URigVMController* RootController = EditorData->GetRigVMClient()->GetController(EditorData->RootGraph);
-	URigVMUnitNode* MainEntryPointNode = RootController->AddUnitNode(FRigUnit_DataInterfaceBeginExecution::StaticStruct(), FRigUnit::GetMethodName(), FVector2D(-400.0f, 0.0f), FString(), false);
+	URigVMController* Controller = EditorData->GetRigVMClient()->GetController(EditorData->RootGraph);
+
+	URigVMUnitNode* MainEntryPointNode = Controller->AddUnitNode(FRigUnit_DataInterfaceBeginExecution::StaticStruct(), FRigUnit::GetMethodName(), FVector2D(-400.0f, 0.0f), FString(), false);
 	URigVMPin* BeginExecutePin = MainEntryPointNode->FindPin(GET_MEMBER_NAME_STRING_CHECKED(FRigUnit_DataInterfaceBeginExecution, ExecuteContext));
 	check(BeginExecutePin);
 	check(BeginExecutePin->GetDirection() == ERigVMPinDirection::Output);
-	
+
 	// Add function to function lib
-	URigVMController* FunctionLibraryController = EditorData->GetRigVMClient()->GetController(EditorData->GetRigVMClient()->GetFunctionLibrary());
-	EditorData->EntryPoint = FunctionLibraryController->AddFunctionToLibrary(UE::DataInterfaceGraph::EntryPointName, true, FVector2D::ZeroVector, false, false);
-	
-	// Add exposed pin for result to function
-	URigVMController* EntryPointController = EditorData->GetRigVMClient()->GetOrCreateController(EditorData->EntryPointGraph);
-	
-	// TODO: using float for now, but needs to be user-driven
 	{
-		const FString CPPType = TEXT("float");
-		const FName CPPTypeObjectPath = NAME_None;
-		const FString DefaultValue = TEXT("0.0");
-		
-		EntryPointController->AddExposedPin(UE::DataInterfaceGraph::ResultName, ERigVMPinDirection::Output, CPPType, CPPTypeObjectPath, DefaultValue, false, false);
+		FRigVMControllerGraphGuard LibraryGraphGuard(Controller, EditorData->GetRigVMClient()->GetFunctionLibrary());
+		EditorData->EntryPoint = Controller->AddFunctionToLibrary(UE::DataInterfaceGraph::EntryPointName, true, FVector2D::ZeroVector, false, false);
+
+		// Add exposed pin for result to function
+		{
+			FRigVMControllerGraphGuard FunctionGraphGuard(Controller, EditorData->EntryPoint->GetContainedGraph());
+
+			// TODO: using float for now, but needs to be user-driven
+			const FString CPPType = TEXT("float");
+			const FName CPPTypeObjectPath = NAME_None;
+			const FString DefaultValue = TEXT("0.0");
+
+			//FunctionController->AddExposedPin(UE::DataInterfaceGraph::ResultName, ERigVMPinDirection::Output, CPPType, CPPTypeObjectPath, DefaultValue, false, false);
+			Controller->AddExposedPin(UE::DataInterfaceGraph::ResultName, ERigVMPinDirection::Output, CPPType, CPPTypeObjectPath, DefaultValue, false, false);
+		}
 	}
 
 	// Add function call
-	URigVMFunctionReferenceNode* FunctionCallNode = RootController->AddFunctionReferenceNode(EditorData->EntryPoint, FVector2D::ZeroVector, UE::DataInterfaceGraph::EntryPointName.ToString(), false, false);
+	URigVMFunctionReferenceNode* FunctionCallNode = Controller->AddFunctionReferenceNode(EditorData->EntryPoint, FVector2D::ZeroVector, UE::DataInterfaceGraph::EntryPointName.ToString(), false, false);
 	URigVMPin* FunctionExecutePin = FunctionCallNode->FindPin(FRigVMStruct::ExecuteContextName.ToString());
 	check(FunctionExecutePin);
 	check(FunctionExecutePin->GetDirection() == ERigVMPinDirection::IO);
 
 	// Link entry point to function call
-	RootController->AddLink(BeginExecutePin, FunctionExecutePin, false); 
-	
+	Controller->AddLink(BeginExecutePin, FunctionExecutePin, false);
+
 	URigVMPin* FunctionResultPin = FunctionCallNode->FindPin(UE::DataInterfaceGraph::ResultName.ToString());
 	check(FunctionResultPin);
 	check(FunctionResultPin->GetDirection() == ERigVMPinDirection::Output);
-	
+
 	// Add end-execution unit of the correct type
 	// TODO: using a float for now, but need to use a registry to determine correct type
-	URigVMUnitNode* MainExitPointNode = RootController->AddUnitNode(FRigUnit_DataInterfaceEndExecution_Float::StaticStruct(), FRigUnit::GetMethodName(), FVector2D(400.0f, 0.0f), FString(), false);
+	URigVMUnitNode* MainExitPointNode = Controller->AddUnitNode(FRigUnit_DataInterfaceEndExecution_Float::StaticStruct(), FRigUnit::GetMethodName(), FVector2D(400.0f, 0.0f), FString(), false);
 	URigVMPin* EndExecutePin = MainExitPointNode->FindPin(GET_MEMBER_NAME_STRING_CHECKED(FRigUnit_DataInterfaceEndExecution, ExecuteContext));
 	check(EndExecutePin);
 	check(EndExecutePin->GetDirection() == ERigVMPinDirection::Input);
@@ -79,13 +83,13 @@ UObject* UDataInterfaceGraphFactory::FactoryCreateNew(UClass* Class, UObject* In
 	URigVMPin* EndExecuteResultPin = MainExitPointNode->FindPin(GET_MEMBER_NAME_STRING_CHECKED(FRigUnit_DataInterfaceEndExecution_Float, Result));
 	check(EndExecuteResultPin);
 	check(EndExecuteResultPin->GetDirection() == ERigVMPinDirection::Input);
-	
+
 	// Link function call to exit point 
-	RootController->AddLink(FunctionExecutePin, EndExecutePin, false); 
+	Controller->AddLink(FunctionExecutePin, EndExecutePin, false);
 
 	// Link result pins 
-	RootController->AddLink(FunctionResultPin, EndExecuteResultPin, false); 
-	
+	Controller->AddLink(FunctionResultPin, EndExecuteResultPin, false);
+
 	// Compile the initial skeleton
 	UE::DataInterfaceGraphUncookedOnly::FUtils::Compile(NewGraph);
 	check(!EditorData->bErrorsDuringCompilation);
diff --git a/Engine/Plugins/Experimental/DataInterfaceGraph/Source/DataInterfaceGraphUncookedOnly/Private/DataInterfaceGraph_EditorData.cpp b/Engine/Plugins/Experimental/DataInterfaceGraph/Source/DataInterfaceGraphUncookedOnly/Private/DataInterfaceGraph_EditorData.cpp
--- a/Engine/Plugins/Experimental/DataInterfaceGraph/Source/DataInterfaceGraphUncookedOnly/Private/DataInterfaceGraph_EditorData.cpp
+++ b/Engine/Plugins/Experimental/DataInterfaceGraph/Source/DataInterfaceGraphUncookedOnly/Private/DataInterfaceGraph_EditorData.cpp
@@ -19,8 +19,8 @@ UDataInterfaceGraph_EditorData::UDataInterfaceGraph_EditorData(const FObjectInit
 	RigVMClient.SetOuterClientHost(this, GET_MEMBER_NAME_CHECKED(UDataInterfaceGraph_EditorData, RigVMClient));
 	{
 		TGuardValue<bool> DisableClientNotifs(RigVMClient.bSuspendNotifications, true);
-		RigVMClient.AddModel(TEXT("RigVMGraph"), false);
-		RigVMClient.GetOrCreateFunctionLibrary(false);
+		RigVMClient.AddModel(TEXT("RigVMGraph"), false, &ObjectInitializer);
+		RigVMClient.GetOrCreateFunctionLibrary(false, &ObjectInitializer);
 	}
 	
 	auto MakeEdGraph = [this, &ObjectInitializer](FName InName) -> UDataInterfaceGraph_EdGraph*
@@ -28,7 +28,7 @@ UDataInterfaceGraph_EditorData::UDataInterfaceGraph_EditorData(const FObjectInit
 		UDataInterfaceGraph_EdGraph* EdGraph = ObjectInitializer.CreateDefaultSubobject<UDataInterfaceGraph_EdGraph>(this, InName);
 		EdGraph->Schema = UDataInterfaceGraph_EdGraphSchema::StaticClass();
 		EdGraph->bAllowRenaming = false;
-		EdGraph->bEditable = false;
+		EdGraph->bEditable = true;
 		EdGraph->bAllowDeletion = false;
 		EdGraph->bIsFunctionDefinition = false;
 		EdGraph->Initialize(this);
@@ -243,9 +243,32 @@ void UDataInterfaceGraph_EditorData::HandleModifiedEvent(ERigVMGraphNotifType In
 	case ERigVMGraphNotifType::LinkRemoved:
 	case ERigVMGraphNotifType::PinArraySizeChanged:
 	case ERigVMGraphNotifType::PinDirectionChanged:
+		{
+			if (InGraph)
+			{
+				InGraph->ClearAST();
+			}
+			break;
+		}
+
 	case ERigVMGraphNotifType::PinDefaultValueChanged:
 		{
-			RequestAutoVMRecompilation();
+			if (InGraph->GetRuntimeAST().IsValid())
+			{
+				URigVMPin* RootPin = CastChecked<URigVMPin>(InSubject)->GetRootPin();
+				FRigVMASTProxy RootPinProxy = FRigVMASTProxy::MakeFromUObject(RootPin);
+				const FRigVMExprAST* Expression = InGraph->GetRuntimeAST()->GetExprForSubject(RootPinProxy);
+				if (Expression == nullptr)
+				{
+					InGraph->ClearAST();
+					break;
+				}
+				else if (Expression->NumParents() > 1)
+				{
+					InGraph->ClearAST();
+					break;
+				}
+			}
 			break;
 		}
 	}
