Fix for DataInterfaceGraph after RigVM execute context refactor

diff --git a/Engine/Plugins/Experimental/DataInterfaceGraph/Source/DataInterfaceGraph/Internal/RigUnit_DataInterfaceEndExecution.h b/Engine/Plugins/Experimental/DataInterfaceGraph/Source/DataInterfaceGraph/Internal/RigUnit_DataInterfaceEndExecution.h
--- a/Engine/Plugins/Experimental/DataInterfaceGraph/Source/DataInterfaceGraph/Internal/RigUnit_DataInterfaceEndExecution.h
+++ b/Engine/Plugins/Experimental/DataInterfaceGraph/Source/DataInterfaceGraph/Internal/RigUnit_DataInterfaceEndExecution.h
@@ -11,7 +11,7 @@
 struct FRigUnitContext;
 
 /** Event for writing back calculated results to external variables */
-USTRUCT(meta=(DisplayName="End Execute Data Interface"))
+USTRUCT(meta = (DisplayName = "End Execute Data Interface", Category = "Events", TitleColor = "1 0 0", NodeColor = "1 1 1"))
 struct DATAINTERFACEGRAPH_API FRigUnit_DataInterfaceEndExecution : public FRigUnit_DataInterfaceBase
 {
 	GENERATED_BODY()
@@ -30,7 +30,7 @@ public:
 };
 
 /** Event for writing back a calculated bool */
-USTRUCT(meta=(DisplayName="End Execute Data Interface Bool"))
+USTRUCT(meta = (DisplayName = "End Execute Data Interface Bool", Category = "Events", TitleColor = "1 0 0", NodeColor = "1 1 1"))
 struct DATAINTERFACEGRAPH_API FRigUnit_DataInterfaceEndExecution_Bool : public FRigUnit_DataInterfaceEndExecution
 {
 	GENERATED_BODY()
@@ -48,7 +48,7 @@ struct DATAINTERFACEGRAPH_API FRigUnit_DataInterfaceEndExecution_Bool : public F
 };
 
 /** Event for writing back a calculated float */
-USTRUCT(meta=(DisplayName="End Execute Data Interface Float"))
+USTRUCT(meta = (DisplayName = "End Execute Data Interface Float", Category = "Events", TitleColor = "1 0 0", NodeColor = "1 1 1"))
 struct DATAINTERFACEGRAPH_API FRigUnit_DataInterfaceEndExecution_Float : public FRigUnit_DataInterfaceEndExecution
 {
 	GENERATED_BODY()
diff --git a/Engine/Plugins/Experimental/DataInterfaceGraph/Source/DataInterfaceGraph/Private/RigUnit_DataInterfaceParameter.h b/Engine/Plugins/Experimental/DataInterfaceGraph/Source/DataInterfaceGraph/Private/RigUnit_DataInterfaceParameter.h
--- a/Engine/Plugins/Experimental/DataInterfaceGraph/Source/DataInterfaceGraph/Private/RigUnit_DataInterfaceParameter.h
+++ b/Engine/Plugins/Experimental/DataInterfaceGraph/Source/DataInterfaceGraph/Private/RigUnit_DataInterfaceParameter.h
@@ -11,7 +11,7 @@ struct FDataInterfaceUnitContext;
 class UDataInterfaceGraph;
 
 /** Unit for reading parameters from context */
-USTRUCT(meta=(DisplayName="Parameter"))
+USTRUCT(meta = (DisplayName = "Parameter", Category = "Parameters", TitleColor = "1 0 0", NodeColor = "1 1 1"))
 struct FRigUnit_DataInterfaceParameter : public FRigUnit_DataInterfaceBase
 {
 	GENERATED_BODY()
@@ -24,7 +24,7 @@ protected:
 };
 
 /** Unit for reading float parameter from context */
-USTRUCT(meta=(DisplayName="Float Parameter"))
+USTRUCT(meta = (DisplayName = "Float Parameter", Category = "Parameters", TitleColor = "1 0 0", NodeColor = "1 1 1"))
 struct FRigUnit_DataInterfaceParameter_Float : public FRigUnit_DataInterfaceParameter
 {
 	GENERATED_BODY()
@@ -38,7 +38,7 @@ protected:
 };
 
 /** Unit for reading data interface graph parameter from context */
-USTRUCT(meta=(DisplayName="Data Interface Parameter"))
+USTRUCT(meta = (DisplayName = "Data Interface Parameter", Category = "Parameters", TitleColor = "1 0 0", NodeColor = "1 1 1"))
 struct FRigUnit_DataInterfaceParameter_DataInterface : public FRigUnit_DataInterfaceParameter
 {
 	GENERATED_BODY()
@@ -142,7 +142,7 @@ struct FDataInterfaceParameter_DataInterface : public FDataInterfaceParameter
 };
 
 /** Base unit for calling data interfaces from graphs */
-USTRUCT(meta=(DisplayName="Data Interface"))
+USTRUCT(meta = (DisplayName = "Data Interface", Category = "Execution", TitleColor = "1 0 0", NodeColor = "1 1 1"))
 struct FRigUnit_DataInterface : public FRigUnit_DataInterfaceBase
 {
 	GENERATED_BODY()
@@ -153,7 +153,7 @@ protected:
 };
 
 /** Unit for getting a float via a data interface */
-USTRUCT(meta=(DisplayName="Data Interface Float"))
+USTRUCT(meta = (DisplayName = "Data Interface Float", Category = "Parameters", TitleColor = "1 0 0", NodeColor = "1 1 1"))
 struct FRigUnit_DataInterface_Float : public FRigUnit_DataInterface
 {
 	GENERATED_BODY()
@@ -191,7 +191,7 @@ protected:
 };
 */
 /** Unit for getting a pose via a data interface */
-USTRUCT(meta=(DisplayName="Float Operator"))
+USTRUCT(meta = (DisplayName = "Float Operator", Category = "Operators", TitleColor = "1 0 0", NodeColor = "1 1 1"))
 struct FRigUnit_FloatOperator : public FRigUnit_DataInterfaceBase
 {
 	GENERATED_BODY()
@@ -220,7 +220,7 @@ protected:
 };
 
 /** Unit for getting a pose via a data interface */
-USTRUCT(meta=(DisplayName="Pose Operator"))
+USTRUCT(meta = (DisplayName = "Pose Operator", Category = "Operators", TitleColor = "1 0 0", NodeColor = "1 1 1"))
 struct FRigUnit_PoseOperator : public FRigUnit_DataInterfaceBase
 {
 	GENERATED_BODY()
@@ -258,7 +258,7 @@ struct FAnimSequenceParameters
 };
 
 /** Unit for getting a pose via a data interface */
-USTRUCT(meta=(DisplayName="Anim Sequence"))
+USTRUCT(meta = (DisplayName = "Anim Sequence", Category = "Animation", TitleColor = "1 0 0", NodeColor = "1 1 1"))
 struct FRigUnit_DataInterface_SequencePlayer : public FRigUnit_DataInterfaceBase
 {
 	GENERATED_BODY()
@@ -278,7 +278,7 @@ protected:
 };
 
 /** Unit for getting a pose via a data interface */
-USTRUCT(meta=(DisplayName="Test Float State - Spring Damper Smoothing", Varying))
+USTRUCT(meta = (DisplayName = "Test Float State - Spring Damper Smoothing", Varying, Category = "Animation", TitleColor = "1 0 0", NodeColor = "1 1 1"))
 struct FRigUnit_TestFloatState : public FRigUnit_DataInterfaceBase
 {
 	GENERATED_BODY()
@@ -301,4 +301,4 @@ protected:
 
 	UPROPERTY(EditAnywhere, Category = "Parameters", meta=(Output))
 	float Result = 0.0f;
-};
\ No newline at end of file
+};
diff --git a/Engine/Plugins/Experimental/DataInterfaceGraph/Source/DataInterfaceGraphEditor/Private/DataInterfaceGraphFactory.cpp b/Engine/Plugins/Experimental/DataInterfaceGraph/Source/DataInterfaceGraphEditor/Private/DataInterfaceGraphFactory.cpp
--- a/Engine/Plugins/Experimental/DataInterfaceGraph/Source/DataInterfaceGraphEditor/Private/DataInterfaceGraphFactory.cpp
+++ b/Engine/Plugins/Experimental/DataInterfaceGraph/Source/DataInterfaceGraphEditor/Private/DataInterfaceGraphFactory.cpp
@@ -32,6 +32,7 @@ UObject* UDataInterfaceGraphFactory::FactoryCreateNew(UClass* Class, UObject* In
 	UDataInterfaceGraph_EditorData* EditorData = NewObject<UDataInterfaceGraph_EditorData>(NewGraph, TEXT("EditorData"));
 	NewGraph->EditorData = EditorData;
 	EditorData->Initialize(/*bRecompileVM*/false);
+	EditorData->GetRigVMClient()->SetExecuteContextStruct(FDataInterfaceExecuteContext::StaticStruct());
 
 	// Add initial execution unit
 	URigVMController* Controller = EditorData->GetRigVMClient()->GetController(EditorData->RootGraph);
diff --git a/Engine/Plugins/Experimental/DataInterfaceGraph/Source/DataInterfaceGraphEditor/Private/SDataInterfaceGraphEditorActionMenu.cpp b/Engine/Plugins/Experimental/DataInterfaceGraph/Source/DataInterfaceGraphEditor/Private/SDataInterfaceGraphEditorActionMenu.cpp
--- a/Engine/Plugins/Experimental/DataInterfaceGraph/Source/DataInterfaceGraphEditor/Private/SDataInterfaceGraphEditorActionMenu.cpp
+++ b/Engine/Plugins/Experimental/DataInterfaceGraph/Source/DataInterfaceGraphEditor/Private/SDataInterfaceGraphEditorActionMenu.cpp
@@ -16,6 +16,7 @@
 #include "RigVMCore/RigVMRegistry.h"
 #include "Units/RigUnit.h"
 #include "Widgets/SToolTip.h"
+#include "DataInterfaceExecuteContext.h"
 
 #define LOCTEXT_NAMESPACE "DataInterfaceEditor"
 
@@ -24,10 +25,26 @@ namespace UE::DataInterfaceGraphEditor
 
 static void CollectAllDataInterfaceActions(FGraphContextMenuBuilder& MenuBuilder)
 {
-	for(const FRigVMFunction& Function : FRigVMRegistry::Get().GetFunctions())
+	static const TArray<UScriptStruct*> AllowedExecuteContexts =
 	{
+		FRigVMExecuteContext::StaticStruct(),
+		FDataInterfaceExecuteContext::StaticStruct()
+	};
+
+	const TChunkedArray<FRigVMFunction>& Functions = FRigVMRegistry::Get().GetFunctions();
+	const int32 NumFunctions = Functions.Num();
+	for(int i=0; i < NumFunctions; i++)
+	{
+		const FRigVMFunction& Function = Functions[i];
+
+		const UScriptStruct* FunctionContext = Function.GetExecuteContextStruct();
+		if (FunctionContext == nullptr || !AllowedExecuteContexts.Contains(FunctionContext))
+		{
+			continue;
+		}
+
 		UScriptStruct* Struct = Function.Struct;
-		if (!Struct->IsChildOf(FRigUnit::StaticStruct()))
+		if (Struct == nullptr)
 		{
 			continue;
 		}
@@ -40,9 +57,14 @@ static void CollectAllDataInterfaceActions(FGraphContextMenuBuilder& MenuBuilder
 		{
 			MenuDescSuffixMetadata = TEXT(" ") + MenuDescSuffixMetadata;
 		}
-		FText NodeCategory = FText::FromString(CategoryMetadata);
-		FText MenuDesc = FText::FromString(DisplayNameMetadata + MenuDescSuffixMetadata);
-		FText ToolTip = Struct->GetToolTipText();
+		const FText NodeCategory = FText::FromString(CategoryMetadata);
+		const FText MenuDesc = FText::FromString(DisplayNameMetadata + MenuDescSuffixMetadata);
+		const FText ToolTip = Struct->GetToolTipText();
+
+		if (MenuDesc.IsEmpty())
+		{
+			continue;
+		}
 
 		MenuBuilder.AddAction(MakeShared<FDataInterfaceGraphSchemaAction_RigUnit>(Struct, NodeCategory, MenuDesc, ToolTip));
 	};
diff --git a/Engine/Plugins/Experimental/DataInterfaceGraph/Source/DataInterfaceGraphUncookedOnly/Internal/DataInterfaceGraph_EditorData.h b/Engine/Plugins/Experimental/DataInterfaceGraph/Source/DataInterfaceGraphUncookedOnly/Internal/DataInterfaceGraph_EditorData.h
--- a/Engine/Plugins/Experimental/DataInterfaceGraph/Source/DataInterfaceGraphUncookedOnly/Internal/DataInterfaceGraph_EditorData.h
+++ b/Engine/Plugins/Experimental/DataInterfaceGraph/Source/DataInterfaceGraphUncookedOnly/Internal/DataInterfaceGraph_EditorData.h
@@ -5,6 +5,7 @@
 #include "CoreMinimal.h"
 #include "RigVMModel/RigVMGraph.h"
 #include "DataInterfaceGraph_EdGraph.h"
+#include "RigVMCore/RigVMGraphFunctionHost.h"
 #include "DataInterfaceGraph_EditorData.generated.h"
 
 class UDataInterfaceGraph;
@@ -20,7 +21,7 @@ namespace UE::DataInterfaceGraphEditor
 }
 
 UCLASS(MinimalAPI)
-class UDataInterfaceGraph_EditorData : public UObject, public IRigVMClientHost
+class UDataInterfaceGraph_EditorData : public UObject, public IRigVMClientHost, public IRigVMGraphFunctionHost
 {
 	GENERATED_BODY()
 
@@ -40,14 +41,18 @@ class UDataInterfaceGraph_EditorData : public UObject, public IRigVMClientHost
 	// IRigVMClientHost interface
 	virtual FRigVMClient* GetRigVMClient() override;
 	virtual const FRigVMClient* GetRigVMClient() const override;
-	virtual IRigVMGraphFunctionHost* GetRigVMGraphFunctionHost() override {return nullptr; }
-	virtual const IRigVMGraphFunctionHost* GetRigVMGraphFunctionHost() const override {return nullptr; }
+	virtual IRigVMGraphFunctionHost* GetRigVMGraphFunctionHost() override;
+	virtual const IRigVMGraphFunctionHost* GetRigVMGraphFunctionHost() const override;
 	virtual void HandleRigVMGraphAdded(const FRigVMClient* InClient, const FString& InNodePath) override;
 	virtual void HandleRigVMGraphRemoved(const FRigVMClient* InClient, const FString& InNodePath) override {}
 	virtual void HandleRigVMGraphRenamed(const FRigVMClient* InClient, const FString& InOldNodePath, const FString& InNewNodePath) override {}
 	virtual void HandleConfigureRigVMController(const FRigVMClient* InClient, URigVMController* InControllerToConfigure) override;
 	virtual UObject* GetEditorObjectForRigVMGraph(URigVMGraph* InVMGraph) const override;
 
+	// IRigVMGraphFunctionHost interface
+	virtual FRigVMGraphFunctionStore* GetRigVMGraphFunctionStore() override;
+	virtual const FRigVMGraphFunctionStore* GetRigVMGraphFunctionStore() const override;
+
 	DATAINTERFACEGRAPHUNCOOKEDONLY_API void Initialize(bool bRecompileVM);
 
 	void RecompileVM();
@@ -74,6 +79,9 @@ class UDataInterfaceGraph_EditorData : public UObject, public IRigVMClientHost
 	UPROPERTY()
 	FRigVMClient RigVMClient;
 
+	UPROPERTY()
+	FRigVMGraphFunctionStore GraphFunctionStore;
+
 	UPROPERTY()
 	TObjectPtr<URigVMGraph> RigVMGraph_DEPRECATED;
 	
diff --git a/Engine/Plugins/Experimental/DataInterfaceGraph/Source/DataInterfaceGraphUncookedOnly/Private/DataInterfaceGraphUncookedOnlyUtils.cpp b/Engine/Plugins/Experimental/DataInterfaceGraph/Source/DataInterfaceGraphUncookedOnly/Private/DataInterfaceGraphUncookedOnlyUtils.cpp
--- a/Engine/Plugins/Experimental/DataInterfaceGraph/Source/DataInterfaceGraphUncookedOnly/Private/DataInterfaceGraphUncookedOnlyUtils.cpp
+++ b/Engine/Plugins/Experimental/DataInterfaceGraph/Source/DataInterfaceGraphUncookedOnly/Private/DataInterfaceGraphUncookedOnlyUtils.cpp
@@ -49,7 +49,8 @@ void FUtils::Compile(UDataInterfaceGraph* InGraph)
 	EditorData->CompileLog.NumErrors = EditorData->CompileLog.NumWarnings = 0;
 
 	URigVMCompiler* Compiler = URigVMCompiler::StaticClass()->GetDefaultObject<URigVMCompiler>();
-	Compiler->Settings = (EditorData->bCompileInDebugMode) ? FRigVMCompileSettings::Fast() : EditorData->VMCompileSettings;
+	EditorData->VMCompileSettings.SetExecuteContextStruct(EditorData->RigVMClient.GetExecuteContextStruct());
+	Compiler->Settings = (EditorData->bCompileInDebugMode) ? FRigVMCompileSettings::Fast(EditorData->VMCompileSettings.GetExecuteContextStruct()) : EditorData->VMCompileSettings;
 	URigVMController* RootController = EditorData->GetRigVMClient()->GetOrCreateController(EditorData->GetRigVMClient()->GetDefaultModel());
 	Compiler->Compile(EditorData->GetRigVMClient()->GetAllModels(false, false), RootController, InGraph->RigVM, InGraph->GetRigVMExternalVariables(), &EditorData->PinToOperandMap);
 
diff --git a/Engine/Plugins/Experimental/DataInterfaceGraph/Source/DataInterfaceGraphUncookedOnly/Private/DataInterfaceGraph_EditorData.cpp b/Engine/Plugins/Experimental/DataInterfaceGraph/Source/DataInterfaceGraphUncookedOnly/Private/DataInterfaceGraph_EditorData.cpp
--- a/Engine/Plugins/Experimental/DataInterfaceGraph/Source/DataInterfaceGraphUncookedOnly/Private/DataInterfaceGraph_EditorData.cpp
+++ b/Engine/Plugins/Experimental/DataInterfaceGraph/Source/DataInterfaceGraphUncookedOnly/Private/DataInterfaceGraph_EditorData.cpp
@@ -105,6 +105,17 @@ const FRigVMClient* UDataInterfaceGraph_EditorData::GetRigVMClient() const
 	return &RigVMClient;
 }
 
+IRigVMGraphFunctionHost* UDataInterfaceGraph_EditorData::GetRigVMGraphFunctionHost()
+{
+	return this;
+}
+
+const IRigVMGraphFunctionHost* UDataInterfaceGraph_EditorData::GetRigVMGraphFunctionHost() const
+{
+	return this;
+}
+
+
 void UDataInterfaceGraph_EditorData::HandleRigVMGraphAdded(const FRigVMClient* InClient, const FString& InNodePath)
 {
 	if(URigVMGraph* RigVMGraph = InClient->GetModel(InNodePath))
@@ -202,6 +213,17 @@ UObject* UDataInterfaceGraph_EditorData::GetEditorObjectForRigVMGraph(URigVMGrap
 	return nullptr;
 }
 
+FRigVMGraphFunctionStore* UDataInterfaceGraph_EditorData::GetRigVMGraphFunctionStore()
+{
+	return &GraphFunctionStore;
+}
+
+const FRigVMGraphFunctionStore* UDataInterfaceGraph_EditorData::GetRigVMGraphFunctionStore() const
+{
+	return &GraphFunctionStore;
+}
+
+
 void UDataInterfaceGraph_EditorData::RecompileVM()
 {
 	UE::DataInterfaceGraphUncookedOnly::FUtils::Compile(GetTypedOuter<UDataInterfaceGraph>());
